Fix NonUnity build

#rnx

diff --git a/Engine/Plugins/MovieScene/MoviePipelineMaskRenderPass/Source/MoviePipelineMaskRenderPass/Private/MoviePipelineObjectIdPass.cpp b/Engine/Plugins/MovieScene/MoviePipelineMaskRenderPass/Source/MoviePipelineMaskRenderPass/Private/MoviePipelineObjectIdPass.cpp
--- a/Engine/Plugins/MovieScene/MoviePipelineMaskRenderPass/Source/MoviePipelineMaskRenderPass/Private/MoviePipelineObjectIdPass.cpp
+++ b/Engine/Plugins/MovieScene/MoviePipelineMaskRenderPass/Source/MoviePipelineMaskRenderPass/Private/MoviePipelineObjectIdPass.cpp
@@ -21,6 +21,7 @@
 #include "Materials/MaterialInstance.h"
 #include "Materials/MaterialInstanceDynamic.h"
 #include "Components/InstancedStaticMeshComponent.h"
+#include "UObject/UObjectAnnotation.h"
 
 DECLARE_CYCLE_STAT(TEXT("STAT_MoviePipeline_AccumulateMaskSample_TT"), STAT_AccumulateMaskSample_TaskThread, STATGROUP_MoviePipeline);
 namespace UE
