MoviePipeline: Reduce the amount of concurrent memory required for panoramic blending by blending it into the final output texture as soon as each sample is done so that we can free the memory. The blending code is still parallel (each sample is worked on by a different task) and only adding the results to the final image is single threaded again but it's only looping over the pixels to add values together, so it's still relatively fast.

diff --git a/Engine/Plugins/MovieScene/MoviePipelineMaskRenderPass/Source/MoviePipelineMaskRenderPass/Private/MoviePipelinePanoramicBlender.cpp b/Engine/Plugins/MovieScene/MoviePipelineMaskRenderPass/Source/MoviePipelineMaskRenderPass/Private/MoviePipelinePanoramicBlender.cpp
--- a/Engine/Plugins/MovieScene/MoviePipelineMaskRenderPass/Source/MoviePipelineMaskRenderPass/Private/MoviePipelinePanoramicBlender.cpp
+++ b/Engine/Plugins/MovieScene/MoviePipelineMaskRenderPass/Source/MoviePipelineMaskRenderPass/Private/MoviePipelinePanoramicBlender.cpp
@@ -9,7 +9,7 @@ FMoviePipelinePanoramicBlender::FMoviePipelinePanoramicBlender(TSharedPtr<MovieP
 	OutputEquirectangularMapSize = InOutputResolution;
 }
 
-static FLinearColor GetColorBilinearFiltered(const FImagePixelData* InSampleData, const FVector2D& InSamplePixelCoords, bool& OutClipped, bool iForceAlphaToOpaque = false)
+static FLinearColor GetColorBilinearFiltered(const FImagePixelData* InSampleData, const FVector2D& InSamplePixelCoords, bool& OutClipped, bool bInForceAlphaToOpaque = false)
 {
 	// Pixel coordinates assume that 0.5, 0.5 is the center of the pixel, so we subtract half to make it indexable.
 	const FVector2D PixelCoordinateIndex = InSamplePixelCoords - 0.5f;
@@ -82,7 +82,7 @@ static FLinearColor GetColorBilinearFiltered(const FImagePixelData* InSampleData
 		 + (UpperLeftPixelColor * (1.0f - FracX) + UpperRightPixelColor * FracX) * FracY;
 
 	// Force final color alpha to opaque if requested
-	if (iForceAlphaToOpaque)
+	if (bInForceAlphaToOpaque)
 	{
 		InterpolatedPixelColor.A = 1.0f;
 	}
@@ -168,6 +168,11 @@ void FMoviePipelinePanoramicBlender::OnCompleteRenderPassDataAvailable_AnyThread
 			int32 EyeMultiplier = DataPayload->Pane.EyeIndex == -1 ? 1 : 2;
 			int32 TotalSampleCount = DataPayload->Pane.NumHorizontalSteps * DataPayload->Pane.NumVerticalSteps * EyeMultiplier;
 			OutputFrame->NumSamplesTotal = TotalSampleCount;
+
+			{
+				LLM_SCOPE_BYNAME(TEXT("MoviePipeline/PanoBlendFrameOutput"));
+				OutputFrame->OutputEquirectangularMap.SetNumZeroed(OutputEquirectangularMapSize.X * OutputEquirectangularMapSize.Y * EyeMultiplier);
+			}
 		}
 	}
 
@@ -309,6 +314,67 @@ void FMoviePipelinePanoramicBlender::OnCompleteRenderPassDataAvailable_AnyThread
 
 	BlendDataTarget->BlendEndTime = FPlatformTime::Seconds();
 
+	// Blend the new sample into the output map as soon as possible, so that we can free the temporary memory held
+	// by this sample. This part is single-threaded (with respect to the other tasks).
+	{
+		// Lock access to our output map
+		FScopeLock ScopeLock(&OutputDataMutex);
+
+
+		// For Stereo images, just offset the second eye by the entire output size so it goes on the bottom half
+		// of the image. We've already made the array twice as big as it needs to be (in the case of stereo) and
+		// below EyeIndex 0 (left eye) will go in the first half of the image, eye index 1 in the second.
+		int32 EyeOffset = 0;
+		if (BlendDataTarget->OriginalDataPayload->Pane.EyeIndex != -1)
+		{
+			EyeOffset = (OutputEquirectangularMapSize.X * OutputEquirectangularMapSize.Y) * BlendDataTarget->OriginalDataPayload->Pane.EyeIndex;
+		}
+
+		for (int32 SampleY = 0; SampleY <= BlendDataTarget->PixelHeight; SampleY++)
+		{
+			for (int32 SampleX = 0; SampleX <= BlendDataTarget->PixelWidth; SampleX++)
+			{
+				int32 OriginalX = SampleX + BlendDataTarget->OutputBoundsMin.X;
+				int32 OriginalY = SampleY + BlendDataTarget->OutputBoundsMin.Y;
+				const int32 OutputPixelX = ((OriginalX % OutputEquirectangularMapSize.X) + OutputEquirectangularMapSize.X) % OutputEquirectangularMapSize.X;
+				const int32 OutputPixelY = OriginalY;
+
+				int32 SourceIndex = SampleX + (SampleY * (BlendDataTarget->PixelWidth + 1));
+				int32 DestIndex = OutputPixelX + (OutputPixelY * OutputEquirectangularMapSize.X);
+				OutputFrame->OutputEquirectangularMap[DestIndex + EyeOffset] += BlendDataTarget->Data[SourceIndex];
+			}
+		}
+
+		bool bDebugSamples = DataPayload->SampleState.bWriteSampleToDisk;
+		if (bDebugSamples)
+		{
+			// Write each blended sample to the output as a debug sample so we can inspect the job blending is doing for each pane.
+			// Hack up the debug output name a bit so they're unique.
+			if (BlendDataTarget->OriginalDataPayload->Pane.EyeIndex >= 0)
+			{
+				BlendDataTarget->OriginalDataPayload->Debug_OverrideFilename = FString::Printf(TEXT("/%s_PaneX_%d_PaneY_%dEye_%d-Blended.%d.exr"),
+					*BlendDataTarget->OriginalDataPayload->PassIdentifier.Name, BlendDataTarget->OriginalDataPayload->Pane.HorizontalStepIndex,
+					BlendDataTarget->OriginalDataPayload->Pane.VerticalStepIndex, DataPayload->Pane.EyeIndex, BlendDataTarget->OriginalDataPayload->SampleState.OutputState.OutputFrameNumber);
+			}
+			else
+			{
+				BlendDataTarget->OriginalDataPayload->Debug_OverrideFilename = FString::Printf(TEXT("/%s_PaneX_%d_PaneY_%dEye_%d-Blended.%d.exr"),
+					*BlendDataTarget->OriginalDataPayload->PassIdentifier.Name, BlendDataTarget->OriginalDataPayload->Pane.HorizontalStepIndex,
+					BlendDataTarget->OriginalDataPayload->Pane.VerticalStepIndex, BlendDataTarget->OriginalDataPayload->SampleState.OutputState.OutputFrameNumber);
+			}
+
+			// Now that the sample has been blended pass it (and the memory it owned, we already read from it) to the debug output step.
+			TUniquePtr<TImagePixelData<FLinearColor>> FinalPixelData = MakeUnique<TImagePixelData<FLinearColor>>(FIntPoint(BlendDataTarget->PixelWidth + 1, BlendDataTarget->PixelHeight + 1), TArray64<FLinearColor>(MoveTemp(BlendDataTarget->Data)), BlendDataTarget->OriginalDataPayload);
+			ensure(OutputMerger.IsValid());
+			OutputMerger.Pin()->OnSingleSampleDataAvailable_AnyThread(MoveTemp(FinalPixelData));
+		}
+		else
+		{
+			// Ensure we reset the memory we allocated, to minimize concurrent allocations.
+			BlendDataTarget->Data.Reset();
+		}
+	}
+
 	// This should only be set to true when this thread is really finished with the work.
 	bool bIsLastSample = false;
 	{
@@ -333,129 +399,34 @@ void FMoviePipelinePanoramicBlender::OnCompleteRenderPassDataAvailable_AnyThread
 
 	if (bIsLastSample)
 	{
-		// If this really was the last sample, then we can blend all the data into the output map.
-		double TotalTimeSpentBlending = 0.f;
-		double MinBlendStart = TNumericLimits<double>::Max();
-		double MaxBlendEnd = TNumericLimits<double>::Lowest();
-
-		{
-			LLM_SCOPE_BYNAME(TEXT("MoviePipeline/PanoBlendFrameOutput"));
-			int32 EyeMultiplier = DataPayload->Pane.EyeIndex == -1 ? 1 : 2;
-			OutputFrame->OutputEquirectangularMap.SetNumZeroed(OutputEquirectangularMapSize.X * OutputEquirectangularMapSize.Y * EyeMultiplier);
-		}
-
-		int32 NumSamplesAccumulated = 0;
-
-		// For each eye
-		for (TPair<int32, TArray<TSharedPtr<FPanoramicBlendData>>>& KVP : OutputFrame->BlendedData)
 		{
-			// For each sample in the eye
-			for (int32 Index = 0; Index < KVP.Value.Num(); Index++)
+			// Now that we've accumulated all the pixel values, we need to normalize them.
+			for (int32 PixelIndex = 0; PixelIndex < OutputFrame->OutputEquirectangularMap.Num(); PixelIndex++)
 			{
-				TSharedPtr<FPanoramicBlendData> BlendData = KVP.Value[Index];
-				TotalTimeSpentBlending += BlendData->BlendEndTime - BlendData->BlendStartTime;
-				if (BlendData->BlendStartTime < MinBlendStart)
-				{
-					MinBlendStart = BlendData->BlendStartTime;
-				}
-
-				if (BlendData->BlendEndTime > MaxBlendEnd)
-				{
-					MaxBlendEnd = BlendData->BlendEndTime;
-				}
-
-				// Write this sample back into the output map. It's a little tricky because we have to handle wraparound. When the array was
-				// first filled, it wrapped around.
-				NumSamplesAccumulated++;
-				// if (NumSamplesAccumulated <= 2)
-				{
-					// For Stereo images, just offset the second eye by the entire output size so it goes on the bottom half
-					// of the image. We've already made the array twice as big as it needs to be (in the case of stereo) and
-					// below EyeIndex 0 (left eye) will go in the first half of the image, eye index 1 in the second.
-					int32 EyeOffset = 0;
-					if (BlendData->OriginalDataPayload->Pane.EyeIndex != -1)
-					{
-						EyeOffset = (OutputEquirectangularMapSize.X * OutputEquirectangularMapSize.Y) * BlendData->OriginalDataPayload->Pane.EyeIndex;
-					}
-
-					for (int32 SampleY = 0; SampleY <= BlendData->PixelHeight; SampleY++)
-					{
-						for (int32 SampleX = 0; SampleX <= BlendData->PixelWidth; SampleX++)
-						{
-							int32 OriginalX = SampleX + BlendData->OutputBoundsMin.X;
-							int32 OriginalY = SampleY + BlendData->OutputBoundsMin.Y;
-							const int32 OutputPixelX = ((OriginalX % OutputEquirectangularMapSize.X) + OutputEquirectangularMapSize.X) % OutputEquirectangularMapSize.X;
-							const int32 OutputPixelY = OriginalY;
-
-							int32 SourceIndex = SampleX + (SampleY * (BlendData->PixelWidth + 1));
-							int32 DestIndex = OutputPixelX + (OutputPixelY * OutputEquirectangularMapSize.X);
-							OutputFrame->OutputEquirectangularMap[DestIndex + EyeOffset] += BlendData->Data[SourceIndex];
-						}
-					}
-				}
-
-				bool bDebugSamples = DataPayload->SampleState.bWriteSampleToDisk;
-				if (bDebugSamples)
-				{
-					// Write each blended sample to the output as a debug sample so we can inspect the job blending is doing for each pane.
-					// Hack up the debug output name a bit so they're unique.
-					if (DataPayload->Pane.EyeIndex >= 0)
-					{
-						BlendData->OriginalDataPayload->Debug_OverrideFilename = FString::Printf(TEXT("/%s_PaneX_%d_PaneY_%dEye_%d-Blended.%d.exr"),
-							*BlendData->OriginalDataPayload->PassIdentifier.Name, BlendData->OriginalDataPayload->Pane.HorizontalStepIndex,
-							BlendData->OriginalDataPayload->Pane.VerticalStepIndex, DataPayload->Pane.EyeIndex, BlendData->OriginalDataPayload->SampleState.OutputState.OutputFrameNumber);
-					}
-					else
-					{
-						BlendData->OriginalDataPayload->Debug_OverrideFilename = FString::Printf(TEXT("/%s_PaneX_%d_PaneY_%dEye_%d-Blended.%d.exr"),
-							*BlendData->OriginalDataPayload->PassIdentifier.Name, BlendData->OriginalDataPayload->Pane.HorizontalStepIndex,
-							BlendData->OriginalDataPayload->Pane.VerticalStepIndex, BlendData->OriginalDataPayload->SampleState.OutputState.OutputFrameNumber);
-					}
-
-					// Now that the sample has been blended pass it (and the memory it owned, we already read from it) to the debug output step.
-					TUniquePtr<TImagePixelData<FLinearColor>> FinalPixelData = MakeUnique<TImagePixelData<FLinearColor>>(FIntPoint(BlendData->PixelWidth + 1, BlendData->PixelHeight + 1), TArray64<FLinearColor>(MoveTemp(BlendData->Data)), BlendData->OriginalDataPayload);
-					ensure(OutputMerger.IsValid());
-					OutputMerger.Pin()->OnSingleSampleDataAvailable_AnyThread(MoveTemp(FinalPixelData));
-				}
-				else
-				{
-					// Ensure we reset the memory we allocated, to minimize concurrent allocations.
-					BlendData->Data.Reset();
-				}
+				FLinearColor& Pixel = OutputFrame->OutputEquirectangularMap[PixelIndex];
+				Pixel.R /= Pixel.A;
+				Pixel.G /= Pixel.A;
+				Pixel.B /= Pixel.A;
 			}
 		}
+		TSharedRef<FImagePixelDataPayload, ESPMode::ThreadSafe> NewPayload = DataPayload->Copy();
 
-		const double ElapsedWallClock = (MaxBlendEnd - MinBlendStart) * 1000.0f;
-		const double ElapsedCPUClock = TotalTimeSpentBlending * 1000.0f;
+		int32 OutputSizeX = OutputEquirectangularMapSize.X;
+		int32 OutputSizeY = DataPayload->Pane.EyeIndex >= 0 ? OutputEquirectangularMapSize.Y * 2 : OutputEquirectangularMapSize.Y;
+		TUniquePtr<TImagePixelData<FLinearColor>> FinalPixelData = MakeUnique<TImagePixelData<FLinearColor>>(FIntPoint(OutputSizeX, OutputSizeY), TArray64<FLinearColor>(MoveTemp(OutputFrame->OutputEquirectangularMap)), NewPayload);
 
-		// UE_LOG(LogMovieRenderPipeline, Warning, TEXT("Blend Time: %8.2fms (Platform Clock) %8.2fms (Total CPU Time)"), ElapsedWallClock, ElapsedCPUClock);
+		if(ensure(OutputMerger.IsValid()))
 		{
-			{
-				// Now that we've accumulated all the pixel values, we need to normalize them.
-				for (int32 PixelIndex = 0; PixelIndex < OutputFrame->OutputEquirectangularMap.Num(); PixelIndex++)
-				{
-					FLinearColor& Pixel = OutputFrame->OutputEquirectangularMap[PixelIndex];
-					Pixel.R /= Pixel.A;
-					Pixel.G /= Pixel.A;
-					Pixel.B /= Pixel.A;
-				}
-			}
-			TSharedRef<FImagePixelDataPayload, ESPMode::ThreadSafe> NewPayload = DataPayload->Copy();
-
-			int32 OutputSizeX = OutputEquirectangularMapSize.X;
-			int32 OutputSizeY = DataPayload->Pane.EyeIndex >= 0 ? OutputEquirectangularMapSize.Y * 2 : OutputEquirectangularMapSize.Y;
-			TUniquePtr<TImagePixelData<FLinearColor>> FinalPixelData = MakeUnique<TImagePixelData<FLinearColor>>(FIntPoint(OutputSizeX, OutputSizeY), TArray64<FLinearColor>(MoveTemp(OutputFrame->OutputEquirectangularMap)), NewPayload);
-
-			ensure(OutputMerger.IsValid());
 			OutputMerger.Pin()->OnCompleteRenderPassDataAvailable_AnyThread(MoveTemp(FinalPixelData));
-
-			// Remove this frame from the PendingData map so that we don't hold references for the duration of the render.
-			// It's important that we use the manually looked up OutputFrame from PendingData
-			// as PendingData uses the equality operator. Some combinations of temporal sampling + slowmo tracks results in different
-			// original source frame numbers, which would cause the tmap lookup to fail and thus returning an empty frame.
-			PendingData.Remove(DataPayload->SampleState.OutputState);
 		}
+
+		// Remove this frame from the PendingData map so that we don't hold references for the duration of the render.
+		// It's important that we use the manually looked up OutputFrame from PendingData
+		// as PendingData uses the equality operator. Some combinations of temporal sampling + slowmo tracks results in different
+		// original source frame numbers, which would cause the tmap lookup to fail and thus returning an empty frame.
+		PendingData.Remove(DataPayload->SampleState.OutputState);
 	}
+
 }
 
 void FMoviePipelinePanoramicBlender::OnSingleSampleDataAvailable_AnyThread(TUniquePtr<FImagePixelData>&& InData)
diff --git a/Engine/Plugins/MovieScene/MoviePipelineMaskRenderPass/Source/MoviePipelineMaskRenderPass/Private/MoviePipelinePanoramicBlender.h b/Engine/Plugins/MovieScene/MoviePipelineMaskRenderPass/Source/MoviePipelineMaskRenderPass/Private/MoviePipelinePanoramicBlender.h
--- a/Engine/Plugins/MovieScene/MoviePipelineMaskRenderPass/Source/MoviePipelineMaskRenderPass/Private/MoviePipelinePanoramicBlender.h
+++ b/Engine/Plugins/MovieScene/MoviePipelineMaskRenderPass/Source/MoviePipelineMaskRenderPass/Private/MoviePipelinePanoramicBlender.h
@@ -50,6 +50,7 @@ private:
 
 	/** Mutex that protects adding/updating/removing from PendingData */
 	FCriticalSection GlobalQueueDataMutex;
+	FCriticalSection OutputDataMutex;
 	FIntPoint OutputEquirectangularMapSize;
 
 	TWeakPtr<MoviePipeline::IMoviePipelineOutputMerger> OutputMerger;
