Fixing up more MRQ render paths that used a hardcoded feature level.

diff --git a/Engine/Plugins/MovieScene/MoviePipelineMaskRenderPass/Source/MoviePipelineMaskRenderPass/Private/MoviePipelineObjectIdPass.cpp b/Engine/Plugins/MovieScene/MoviePipelineMaskRenderPass/Source/MoviePipelineMaskRenderPass/Private/MoviePipelineObjectIdPass.cpp
--- a/Engine/Plugins/MovieScene/MoviePipelineMaskRenderPass/Source/MoviePipelineMaskRenderPass/Private/MoviePipelineObjectIdPass.cpp
+++ b/Engine/Plugins/MovieScene/MoviePipelineMaskRenderPass/Source/MoviePipelineMaskRenderPass/Private/MoviePipelineObjectIdPass.cpp
@@ -318,7 +318,7 @@ void UMoviePipelineObjectIdRenderPass::RenderSample_GameThreadImpl(const FMovieP
 	// Submit to be rendered. Main render pass always uses target 0. We do this before making the Hitproxy cache because
 	// BeginRenderingViewFamily ensures render state for things are created.
 	FRenderTarget* RenderTarget = GetViewRenderTarget()->GameThread_GetRenderTargetResource();
-	FCanvas Canvas = FCanvas(RenderTarget, nullptr, GetPipeline()->GetWorld(), ERHIFeatureLevel::SM5, FCanvas::CDM_DeferDrawing, 1.0f);
+	FCanvas Canvas = FCanvas(RenderTarget, nullptr, GetPipeline()->GetWorld(), ViewFamily->GetFeatureLevel(), FCanvas::CDM_DeferDrawing, 1.0f);
 	GetRendererModule().BeginRenderingViewFamily(&Canvas, ViewFamily.Get());
 
 	// The Hitproxy array gets invalidated quite often, so the results are no longer valid in the accumulation thread.
diff --git a/Engine/Plugins/MovieScene/MoviePipelineMaskRenderPass/Source/MoviePipelineMaskRenderPass/Private/MoviePipelinePanoramicPass.cpp b/Engine/Plugins/MovieScene/MoviePipelineMaskRenderPass/Source/MoviePipelineMaskRenderPass/Private/MoviePipelinePanoramicPass.cpp
--- a/Engine/Plugins/MovieScene/MoviePipelineMaskRenderPass/Source/MoviePipelineMaskRenderPass/Private/MoviePipelinePanoramicPass.cpp
+++ b/Engine/Plugins/MovieScene/MoviePipelineMaskRenderPass/Source/MoviePipelineMaskRenderPass/Private/MoviePipelinePanoramicPass.cpp
@@ -478,7 +478,7 @@ void UMoviePipelinePanoramicPass::RenderSample_GameThreadImpl(const FMoviePipeli
 
 				// Submit the view to be rendered.
 				FRenderTarget* RenderTarget = CanvasReadbackTexture->GameThread_GetRenderTargetResource();
-				FCanvas Canvas = FCanvas(RenderTarget, nullptr, GetPipeline()->GetWorld(), ERHIFeatureLevel::SM5, FCanvas::CDM_DeferDrawing, 1.0f);
+				FCanvas Canvas = FCanvas(RenderTarget, nullptr, GetPipeline()->GetWorld(), ViewFamily->GetFeatureLevel(), FCanvas::CDM_DeferDrawing, 1.0f);
 				GetRendererModule().BeginRenderingViewFamily(&Canvas, ViewFamily.Get());
 
 				// Schedule a readback and then high-res accumulation.
