Added Custom Near Clip Distance setting to CineCameraActor.

diff --git a/Engine/Plugins/MovieScene/MoviePipelineMaskRenderPass/Source/MoviePipelineMaskRenderPass/Private/MoviePipelinePanoramicBlender.cpp b/Engine/Plugins/MovieScene/MoviePipelineMaskRenderPass/Source/MoviePipelineMaskRenderPass/Private/MoviePipelinePanoramicBlender.cpp
--- a/Engine/Plugins/MovieScene/MoviePipelineMaskRenderPass/Source/MoviePipelineMaskRenderPass/Private/MoviePipelinePanoramicBlender.cpp
+++ b/Engine/Plugins/MovieScene/MoviePipelineMaskRenderPass/Source/MoviePipelineMaskRenderPass/Private/MoviePipelinePanoramicBlender.cpp
@@ -123,7 +123,7 @@ void FMoviePipelinePanoramicBlender::OnCompleteRenderPassDataAvailable_AnyThread
 	const FVector SampleDirectionOnPhi = FVector(FMath::Cos(SamplePitchRad), 0.f, FMath::Sin(SamplePitchRad));
 
 	// Now construct a projection matrix representing the sample matching the original perspective it was taken from.
-	const FMatrix SampleProjectionMatrix = FReversedZPerspectiveMatrix(FMath::DegreesToRadians(SampleHalfHorizontalFoVDegrees), SampleSize.X, SampleSize.Y, GNearClippingPlane);
+	const FMatrix SampleProjectionMatrix = FReversedZPerspectiveMatrix(FMath::DegreesToRadians(SampleHalfHorizontalFoVDegrees), SampleSize.X, SampleSize.Y, DataPayload->Pane.NearClippingPlane);
 
 	// For our given output size, figure out how many degrees each pixel represents.
 	const float EquiRectMapThetaStep = 360.f / (float)OutputEquirectangularMapSize.X;
diff --git a/Engine/Plugins/MovieScene/MoviePipelineMaskRenderPass/Source/MoviePipelineMaskRenderPass/Private/MoviePipelinePanoramicPass.cpp b/Engine/Plugins/MovieScene/MoviePipelineMaskRenderPass/Source/MoviePipelineMaskRenderPass/Private/MoviePipelinePanoramicPass.cpp
--- a/Engine/Plugins/MovieScene/MoviePipelineMaskRenderPass/Source/MoviePipelineMaskRenderPass/Private/MoviePipelinePanoramicPass.cpp
+++ b/Engine/Plugins/MovieScene/MoviePipelineMaskRenderPass/Source/MoviePipelineMaskRenderPass/Private/MoviePipelinePanoramicPass.cpp
@@ -326,7 +326,13 @@ FSceneView* UMoviePipelinePanoramicPass::GetSceneViewForSampleState(FSceneViewFa
 
 	// Calculate a Projection Matrix
 	{
-		const float MinZ = GNearClippingPlane;
+		float MinZ = GNearClippingPlane;
+		if (LocalPlayerController->PlayerCameraManager)
+		{
+			float NearClipPlane = LocalPlayerController->PlayerCameraManager->GetCameraCacheView().PerspectiveNearClipPlane;
+			MinZ = NearClipPlane > 0 ? NearClipPlane : MinZ;
+		}
+		PanoPane->NearClippingPlane = MinZ;
 		const float MaxZ = MinZ;
 		// Avoid zero ViewFOV's which cause divide by zero's in projection matrix
 		const float MatrixFOV = FMath::Max(0.001f, ViewFOV) * (float)PI / 360.0f;
diff --git a/Engine/Plugins/MovieScene/MoviePipelineMaskRenderPass/Source/MoviePipelineMaskRenderPass/Private/MoviePipelinePanoramicPass.h b/Engine/Plugins/MovieScene/MoviePipelineMaskRenderPass/Source/MoviePipelineMaskRenderPass/Private/MoviePipelinePanoramicPass.h
--- a/Engine/Plugins/MovieScene/MoviePipelineMaskRenderPass/Source/MoviePipelineMaskRenderPass/Private/MoviePipelinePanoramicPass.h
+++ b/Engine/Plugins/MovieScene/MoviePipelineMaskRenderPass/Source/MoviePipelineMaskRenderPass/Private/MoviePipelinePanoramicPass.h
@@ -21,6 +21,8 @@ struct FPanoPane : public UMoviePipelineImagePassBase::IViewCalcPayload
 	FRotator OriginalCameraRotation;
 	// The camera rotation last frame, used to ensure camera motion vectors are right.
 	FRotator PrevOriginalCameraRotation;
+	// The near clip plane distance from the camera.
+	float NearClippingPlane;
 
 	// How far apart are the eyes (total) for stereo?
 	float EyeSeparation;
diff --git a/Engine/Plugins/MovieScene/MovieRenderPipeline/Source/MovieRenderPipelineRenderPasses/Private/MoviePipelineImagePassBase.cpp b/Engine/Plugins/MovieScene/MovieRenderPipeline/Source/MovieRenderPipelineRenderPasses/Private/MoviePipelineImagePassBase.cpp
--- a/Engine/Plugins/MovieScene/MovieRenderPipeline/Source/MovieRenderPipelineRenderPasses/Private/MoviePipelineImagePassBase.cpp
+++ b/Engine/Plugins/MovieScene/MovieRenderPipeline/Source/MovieRenderPipelineRenderPasses/Private/MoviePipelineImagePassBase.cpp
@@ -363,7 +363,12 @@ FSceneView* UMoviePipelineImagePassBase::GetSceneViewForSampleState(FSceneViewFa
 			}
 		}
 
-		const float MinZ = GNearClippingPlane;
+		float MinZ = GNearClippingPlane;
+		if (PlayerCameraManager)
+		{
+			float NearClipPlane = PlayerCameraManager->GetCameraCacheView().PerspectiveNearClipPlane;
+			MinZ = NearClipPlane > 0 ? NearClipPlane : MinZ;
+		}
 		const float MaxZ = MinZ;
 		// Avoid zero ViewFOV's which cause divide by zero's in projection matrix
 		const float MatrixFOV = FMath::Max(0.001f, ViewFOV) * (float)PI / 360.0f;
diff --git a/Engine/Source/Runtime/CinematicCamera/Private/CineCameraComponent.cpp b/Engine/Source/Runtime/CinematicCamera/Private/CineCameraComponent.cpp
--- a/Engine/Source/Runtime/CinematicCamera/Private/CineCameraComponent.cpp
+++ b/Engine/Source/Runtime/CinematicCamera/Private/CineCameraComponent.cpp
@@ -68,6 +68,9 @@ UCineCameraComponent::UCineCameraComponent()
 	static ConstructorHelpers::FObjectFinder<UMaterial> PlaneMat(TEXT("/Engine/EngineDebugMaterials/M_SimpleUnlitTranslucent.M_SimpleUnlitTranslucent"));
 	FocusPlaneVisualizationMaterial = PlaneMat.Object;
 #endif
+
+	CustomNearClippingPlane = GNearClippingPlane;
+	bOverride_CustomNearClippingPlane = false;
 }
 
 void UCineCameraComponent::Serialize(FArchive& Ar)
@@ -455,6 +458,8 @@ void UCineCameraComponent::GetCameraView(float DeltaTime, FMinimalViewInfo& Desi
 
 	UpdateCameraLens(DeltaTime, DesiredView);
 
+	DesiredView.PerspectiveNearClipPlane = bOverride_CustomNearClippingPlane ? CustomNearClippingPlane : -1.0f;
+
 	bResetInterpolation = false;
 }
 
diff --git a/Engine/Source/Runtime/CinematicCamera/Public/CineCameraComponent.h b/Engine/Source/Runtime/CinematicCamera/Public/CineCameraComponent.h
--- a/Engine/Source/Runtime/CinematicCamera/Public/CineCameraComponent.h
+++ b/Engine/Source/Runtime/CinematicCamera/Public/CineCameraComponent.h
@@ -272,6 +272,13 @@ public:
 	UPROPERTY(VisibleAnywhere, BlueprintReadOnly, Category = "Current Camera Settings")
 	float CurrentFocusDistance;
 
+	UPROPERTY(EditAnywhere, Category = "Current Camera Settings", meta = (InlineEditConditionToggle))
+	uint32 bOverride_CustomNearClippingPlane:1;
+
+	/** Set bOverride_CustomNearClippingPlane to true if you want to use a custom clipping plane instead of GNearClippingPlane. */
+	UPROPERTY(EditAnywhere, BlueprintReadWrite, AdvancedDisplay, Category = "Current Camera Settings", meta = (editcondition = "bOverride_CustomNearClippingPlane"))
+	float CustomNearClippingPlane;
+
 #if WITH_EDITORONLY_DATA
 	/** Read-only. Control this value with CurrentFocalLength (and filmback settings). */
 	UPROPERTY(VisibleAnywhere, Category = "Current Camera Settings")
