Fix static analysis warning.

diff --git a/Engine/Plugins/MovieScene/MoviePipelineMaskRenderPass/Source/MoviePipelineMaskRenderPass/Private/MoviePipelinePanoramicPass.cpp b/Engine/Plugins/MovieScene/MoviePipelineMaskRenderPass/Source/MoviePipelineMaskRenderPass/Private/MoviePipelinePanoramicPass.cpp
--- a/Engine/Plugins/MovieScene/MoviePipelineMaskRenderPass/Source/MoviePipelineMaskRenderPass/Private/MoviePipelinePanoramicPass.cpp
+++ b/Engine/Plugins/MovieScene/MoviePipelineMaskRenderPass/Source/MoviePipelineMaskRenderPass/Private/MoviePipelinePanoramicPass.cpp
@@ -327,7 +327,7 @@ FSceneView* UMoviePipelinePanoramicPass::GetSceneViewForSampleState(FSceneViewFa
 	// Calculate a Projection Matrix
 	{
 		float MinZ = GNearClippingPlane;
-		if (LocalPlayerController->PlayerCameraManager)
+		if (LocalPlayerController && LocalPlayerController->PlayerCameraManager)
 		{
 			float NearClipPlane = LocalPlayerController->PlayerCameraManager->GetCameraCacheView().PerspectiveNearClipPlane;
 			MinZ = NearClipPlane > 0 ? NearClipPlane : MinZ;
