MoviePipeline: Panoramic render pass no longer creates a visible seam on the right edge of the resulting panoramic image when using 8 horizontal samples.

Fixed the uplugin descriptor to include the word Panoramic so it returns in search results when searching for panoramic support.

Fixed the Movie Pipeline RenderPass warning about alpha-in-tonemapper even though alpha support is not required to use Panoramic images (and does not support alpha pass through)

Fixed the resulting panoramic image still having blend-weight data in the alpha channel which made pngs semi-transparent.

diff --git a/Engine/Plugins/MovieScene/MoviePipelineMaskRenderPass/MoviePipelineMaskRenderPass.uplugin b/Engine/Plugins/MovieScene/MoviePipelineMaskRenderPass/MoviePipelineMaskRenderPass.uplugin
--- a/Engine/Plugins/MovieScene/MoviePipelineMaskRenderPass/MoviePipelineMaskRenderPass.uplugin
+++ b/Engine/Plugins/MovieScene/MoviePipelineMaskRenderPass/MoviePipelineMaskRenderPass.uplugin
@@ -3,7 +3,7 @@
 	"Version": 1,
 	"VersionName": "1.0",
 	"FriendlyName": "Movie Render Queue Additional Render Passes",
-	"Description": "Additional render passes for the Movie Render Queue which generate mattes with some limitations. Editor Only.",
+	"Description": "Additional render passes for the Movie Render Queue. This currently includes the ObjectId pass (Editor Only) which generates object mattes with some limitations (using the Cryptomatte specification), and a Panoramic pass with better Sequencer integration than the Panoramic Capture plugin.",
 	"Category": "Rendering",
 	"CreatedBy": "Epic Games, Inc.",
 	"CreatedByURL": "",
diff --git a/Engine/Plugins/MovieScene/MoviePipelineMaskRenderPass/Source/MoviePipelineMaskRenderPass/Private/MoviePipelinePanoramicBlender.cpp b/Engine/Plugins/MovieScene/MoviePipelineMaskRenderPass/Source/MoviePipelineMaskRenderPass/Private/MoviePipelinePanoramicBlender.cpp
--- a/Engine/Plugins/MovieScene/MoviePipelineMaskRenderPass/Source/MoviePipelineMaskRenderPass/Private/MoviePipelinePanoramicBlender.cpp
+++ b/Engine/Plugins/MovieScene/MoviePipelineMaskRenderPass/Source/MoviePipelineMaskRenderPass/Private/MoviePipelinePanoramicBlender.cpp
@@ -138,13 +138,13 @@ void FMoviePipelinePanoramicBlender::OnCompleteRenderPassDataAvailable_AnyThread
 	// The MinBound for X can actually be greater than the MaxBound due to wrapping, modulo is applied at eval time to ensure it wraps right.
 	float SampleYawMin = SampleRotation.Yaw - SampleHalfHorizontalFoVDegrees;
 	float SampleYawMax = SampleRotation.Yaw + SampleHalfHorizontalFoVDegrees;
-	int32 PixelIndexHorzMinBound = FMath::FloorToInt(((SampleYawMin)+180.f) / EquiRectMapThetaStep);
-	int32 PixelIndexHorzMaxBound = FMath::FloorToInt(((SampleYawMax)+180.f) / EquiRectMapThetaStep);
+	int32 PixelIndexHorzMinBound = FMath::FloorToInt(((SampleYawMin) + 180.f) / EquiRectMapThetaStep);
+	int32 PixelIndexHorzMaxBound = FMath::FloorToInt(((SampleYawMax) + 180.f) / EquiRectMapThetaStep);
 
 	float SamplePitchMin = FMath::Max(SampleRotation.Pitch - SampleHalfVerticalFoVDegrees, -90.f); // Clamped to [-90, 90]
 	float SamplePitchMax = FMath::Min(SampleRotation.Pitch + SampleHalfVerticalFoVDegrees, 90.f); // Clamped to [-90, 90]
-	int32 PixelIndexVertMinBound = FMath::Max((OutputEquirectangularMapSize.Y - 1) - FMath::FloorToInt((SamplePitchMax + 90.f) / EquiRectMapPhiStep), 0);
-	int32 PixelIndexVertMaxBound = FMath::Min((OutputEquirectangularMapSize.Y - 1) - FMath::FloorToInt((SamplePitchMin + 90.f) / EquiRectMapPhiStep), OutputEquirectangularMapSize.Y - 1);
+	int32 PixelIndexVertMinBound = FMath::Max((OutputEquirectangularMapSize.Y) - FMath::FloorToInt((SamplePitchMax + 90.f) / EquiRectMapPhiStep), 0);
+	int32 PixelIndexVertMaxBound = FMath::Min((OutputEquirectangularMapSize.Y) - FMath::FloorToInt((SamplePitchMin + 90.f) / EquiRectMapPhiStep), OutputEquirectangularMapSize.Y);
 
 	{
 		// Do a quick lock while we're iterating/adding to the PendingData array so a second sample
@@ -200,117 +200,94 @@ void FMoviePipelinePanoramicBlender::OnCompleteRenderPassDataAvailable_AnyThread
 	BlendDataTarget->PixelHeight = BlendDataTarget->OutputBoundsMax.Y - BlendDataTarget->OutputBoundsMin.Y;
 
 	// These need to be zeroed as we don't always touch every pixel in the rect with blending
-	// and they get +='d together later. We allocate extra pixel here because the for loops are <= and not <.
+	// and they get +='d together later.
 	{
 		LLM_SCOPE_BYNAME(TEXT("MoviePipeline/PanoBlendPerTaskOutput"));
-		BlendDataTarget->Data.SetNumZeroed((BlendDataTarget->PixelWidth + 1) * (BlendDataTarget->PixelHeight + 1));
+		BlendDataTarget->Data.SetNumZeroed((BlendDataTarget->PixelWidth) * (BlendDataTarget->PixelHeight));
 	}
 
 	// Finally we can perform our actual blending. We blend into our intermediate buffer
 	// instead of the final output array to avoid multiple threads contending for pixels.
-// The define was confusing the static analysis for some reason, so we'll just comment it out.
-// The overhead of the ParallelFor isn't really worth it here now that we go wide with each sample
-// being accumulated on the task thread into a separate output buffer anyways.
-//#define ACCUMULATE_IN_SCANLINES 0
-//#if ACCUMULATE_IN_SCANLINES
-//	const int32 NumRowsPerTask = 1;
-//	const int32 NumScanLineTasks = FMath::DivideAndRoundUp(PixelIndexVertMaxBound - PixelIndexVertMinBound + 1, NumRowsPerTask);
-//
-//	// Like most of the Movie Pipeline, we parallelize on scanlines.
-//	ParallelFor(NumScanLineTasks,
-//		[&](int32 InScanLineTaskIndex)
-//	{
-//		int32 VertIndexStart = NumRowsPerTask * InScanLineTaskIndex + PixelIndexVertMinBound;
-//		int32 VertIndexEnd = ((NumRowsPerTask * (InScanLineTaskIndex + 1)) - 1) + PixelIndexVertMinBound;
-//		
-//		// Clamp it so we don't go out of bounds on the end rows.
-//		VertIndexEnd = FMath::Min(VertIndexEnd, PixelIndexVertMaxBound);
-//		for (int32 Y = VertIndexStart; Y <= VertIndexEnd; Y++)
-//#else
-		for (int32 Y = PixelIndexVertMinBound; Y <= PixelIndexVertMaxBound; Y++)
-//#endif
+	for (int32 Y = PixelIndexVertMinBound; Y < PixelIndexVertMaxBound; Y++)
+	{
+		for (int32 X = PixelIndexHorzMinBound; X < PixelIndexHorzMaxBound; X++)
 		{
-			for (int32 X = PixelIndexHorzMinBound; X <= PixelIndexHorzMaxBound; X++)
+			// These X, Y coordinates are in output resolution space which is where we want to blend to.
+			// Our X bounds may go OOB, but we wrap horizontally so we need to figure out the proper X index.
+			const int32 OutputPixelX = ((X % OutputEquirectangularMapSize.X) + OutputEquirectangularMapSize.X) % OutputEquirectangularMapSize.X;
+			const int32 OutputPixelY = Y;
+
+			// Get the spherical coordinates (Theta and Phi) corresponding to the X and Y of the equirectangular map coordinates, converted to
+			// [-180, 180] and [-90, 90] coordinate space respectively. The half pixel offset is used to make the center of a pixel be considered
+			// that coordinate, and Phi increments in the opposite direction of Y.
+			const float Theta = EquiRectMapThetaStep * (((float)OutputPixelX) + 0.5f) - 180.f;
+			const float Phi = EquiRectMapPhiStep * (((float)OutputEquirectangularMapSize.Y - OutputPixelY) + 0.5f) - 90.f;
+
+			// Now convert the spherical coordinates into an actual direction (on the output map)
+			const float ThetaDeg = FMath::DegreesToRadians(Theta);
+			const float PhiDeg = FMath::DegreesToRadians(Phi);
+			const FVector OutputDirection(FMath::Cos(PhiDeg) * FMath::Cos(ThetaDeg), FMath::Cos(PhiDeg) * FMath::Sin(ThetaDeg), FMath::Sin(PhiDeg));
+			const FVector OutputDirectionTheta = FVector(FMath::Cos(ThetaDeg), FMath::Sin(ThetaDeg), 0);
+			const FVector OutputDirectionPhi = FVector(FMath::Cos(PhiDeg), 0.f, FMath::Sin(PhiDeg));
+
+			// Now we can compute how much the sample should influence this pixel. It is weighted by angular distance to the direction
+			// so that the edges have less influence (where they'd be more distorted anyways).
+			const float DirectionPhiDot = FVector::DotProduct(OutputDirectionPhi, SampleDirectionOnPhi); // ToDo: This only considers the whole Pano Pane and not per pixel of sample?
+			const float DirectionThetaDot = FVector::DotProduct(OutputDirectionTheta, SampleDirectionOnTheta);
+			const float WeightTheta = FMath::Max(DirectionThetaDot - SampleHalfHorizontalFoVCosine, 0.0f) / (1.0f - SampleHalfHorizontalFoVCosine);
+			const float WeightPhi = FMath::Max(DirectionPhiDot - SampleHalfVerticalFoVCosine, 0.0f) / (1.0f - SampleHalfVerticalFoVCosine);
+
+			const float SampleWeight = WeightTheta * WeightPhi;
+			const float SampleWeightSquared = SampleWeight* SampleWeight; // Exponential falloff produces a nicer blending result.
+
+			// The sample weight may be very small and not worth influencing this pixel.
+			if (SampleWeightSquared > KINDA_SMALL_NUMBER)
 			{
-				// These X, Y coordinates are in output resolution space which is where we want to blend to.
-				// Our X bounds may go OOB, but we wrap horizontally so we need to figure out the proper X index.
-				const int32 OutputPixelX = ((X % OutputEquirectangularMapSize.X) + OutputEquirectangularMapSize.X) % OutputEquirectangularMapSize.X;
-				const int32 OutputPixelY = Y;
-
-				// Get the spherical coordinates (Theta and Phi) corresponding to the X and Y of the equirectangular map coordinates, converted to
-				// [-180, 180] and [-90, 90] coordinate space respectively. The half pixel offset is used to make the center of a pixel be considered
-				// that coordinate, and Phi increments in the opposite direction of Y.
-				const float Theta = EquiRectMapThetaStep * (((float)OutputPixelX) + 0.5f) - 180.f;
-				const float Phi = EquiRectMapPhiStep * (((float)OutputEquirectangularMapSize.Y - OutputPixelY - 1) + 0.5f) - 90.f;
-
-				// Now convert the spherical coordinates into an actual direction (on the output map)
-				const float ThetaDeg = FMath::DegreesToRadians(Theta);
-				const float PhiDeg = FMath::DegreesToRadians(Phi);
-				const FVector OutputDirection(FMath::Cos(PhiDeg) * FMath::Cos(ThetaDeg), FMath::Cos(PhiDeg) * FMath::Sin(ThetaDeg), FMath::Sin(PhiDeg));
-				const FVector OutputDirectionTheta = FVector(FMath::Cos(ThetaDeg), FMath::Sin(ThetaDeg), 0);
-				const FVector OutputDirectionPhi = FVector(FMath::Cos(PhiDeg), 0.f, FMath::Sin(PhiDeg));
-
-				// Now we can compute how much the sample should influence this pixel. It is weighted by angular distance to the direction
-				// so that the edges have less influence (where they'd be more distorted anyways).
-				const float DirectionPhiDot = FVector::DotProduct(OutputDirectionPhi, SampleDirectionOnPhi); // ToDo: This only considers the whole Pano Pane and not per pixel of sample?
-				const float DirectionThetaDot = FVector::DotProduct(OutputDirectionTheta, SampleDirectionOnTheta);
-				const float WeightTheta = FMath::Max(DirectionThetaDot - SampleHalfHorizontalFoVCosine, 0.0f) / (1.0f - SampleHalfHorizontalFoVCosine);
-				const float WeightPhi = FMath::Max(DirectionPhiDot - SampleHalfVerticalFoVCosine, 0.0f) / (1.0f - SampleHalfVerticalFoVCosine);
-
-				const float SampleWeight = WeightTheta * WeightPhi;
-				const float SampleWeightSquared = SampleWeight* SampleWeight; // Exponential falloff produces a nicer blending result.
-
-				// The sample weight may be very small and not worth influencing this pixel.
-				if (SampleWeightSquared > KINDA_SMALL_NUMBER)
+				// Transform the direction vector from the equirectangular map world space to sample world space
+				FVector4 DirectionInSampleWorldSpace = FVector4(SampleRotation.UnrotateVector(OutputDirection), 1.0f);
+
+				static const FMatrix UnrealCoordinateConversion = FMatrix(
+					FPlane(0, 0, 1, 0),
+					FPlane(1, 0, 0, 0),
+					FPlane(0, 1, 0, 0),
+					FPlane(0, 0, 0, 1));
+				DirectionInSampleWorldSpace = UnrealCoordinateConversion.TransformFVector4(DirectionInSampleWorldSpace);
+
+				// Then project that direction into sample clip space
+				FVector4 DirectionInSampleClipSpace = SampleProjectionMatrix.TransformFVector4(DirectionInSampleWorldSpace);
+
+				// Converted into normalized device space (Divide by w for perspective)
+				FVector DirectionInSampleNDSpace = FVector(DirectionInSampleClipSpace) / DirectionInSampleClipSpace.W;
+
+				// Get the final pixel coordinates (direction in screen space)
+				FVector2D DirectionInSampleScreenSpace = ((FVector2D(DirectionInSampleNDSpace) + 1.0f) / 2.0f) * FVector2D(SampleSize.X, SampleSize.Y);
+
+				// Flip the Y value due to Y's zero coordinate being top left.
+				DirectionInSampleScreenSpace.Y = ((float)SampleSize.Y - DirectionInSampleScreenSpace.Y) - 1.0f;
+
+				// Do a bilinear color sample at the pixel coordinates (from the sample), weight it, and add it to the output map. We store
+				// weights separately so that we can preserve the alpha channel of the main image.
+				bool bClipped = false;
+				FLinearColor SampleColor = GetColorBilinearFiltered(InData.Get(), DirectionInSampleScreenSpace, bClipped, true);
+
+				if (!bClipped)
 				{
-					// Transform the direction vector from the equirectangular map world space to sample world space
-					FVector4 DirectionInSampleWorldSpace = FVector4(SampleRotation.UnrotateVector(OutputDirection), 1.0f);
-
-					static const FMatrix UnrealCoordinateConversion = FMatrix(
-						FPlane(0, 0, 1, 0),
-						FPlane(1, 0, 0, 0),
-						FPlane(0, 1, 0, 0),
-						FPlane(0, 0, 0, 1));
-					DirectionInSampleWorldSpace = UnrealCoordinateConversion.TransformFVector4(DirectionInSampleWorldSpace);
-
-					// Then project that direction into sample clip space
-					FVector4 DirectionInSampleClipSpace = SampleProjectionMatrix.TransformFVector4(DirectionInSampleWorldSpace);
-
-					// Converted into normalized device space (Divide by w for perspective)
-					FVector DirectionInSampleNDSpace = FVector(DirectionInSampleClipSpace) / DirectionInSampleClipSpace.W;
-
-					// Get the final pixel coordinates (direction in screen space)
-					FVector2D DirectionInSampleScreenSpace = ((FVector2D(DirectionInSampleNDSpace) + 1.0f) / 2.0f) * FVector2D(SampleSize.X, SampleSize.Y);
-
-					// Flip the Y value due to Y's zero coordinate being top left.
-					DirectionInSampleScreenSpace.Y = ((float)SampleSize.Y - DirectionInSampleScreenSpace.Y) - 1.0f;
-
-					// Do a bilinear color sample at the pixel coordinates (from the sample), weight it, and add it to the output map. We store
-					// weights separately so that we can preserve the alpha channel of the main image.
-					bool bClipped = false;
-					FLinearColor SampleColor = GetColorBilinearFiltered(InData.Get(), DirectionInSampleScreenSpace, bClipped, true);
-
-					if (!bClipped)
-					{
-						// When we calculate the actual output location we need to shift the X/Y. This is because up until now the math has been done in
-						// output resolution space, but each sample only allocates a color map big enough for itself. It'll get shifted back out to the
-						// right location later.
-						int32 SampleOutputX = OutputPixelX - BlendDataTarget->OutputBoundsMin.X;
-						// Mod this again by our output map so we don't OOB on it. It'll wrap weirdly in the output map but should restore fine.
-						SampleOutputX = ((SampleOutputX % (BlendDataTarget->PixelWidth+1)) + (BlendDataTarget->PixelWidth+1)) % (BlendDataTarget->PixelWidth+1); // Positive Mod
-						int32 SampleOutputY = Y;
-						SampleOutputY -= BlendDataTarget->OutputBoundsMin.Y;
+					// When we calculate the actual output location we need to shift the X/Y. This is because up until now the math has been done in
+					// output resolution space, but each sample only allocates a color map big enough for itself. It'll get shifted back out to the
+					// right location later.
+					int32 SampleOutputX = OutputPixelX - BlendDataTarget->OutputBoundsMin.X;
+					// Mod this again by our output map so we don't OOB on it. It'll wrap weirdly in the output map but should restore fine.
+					SampleOutputX = ((SampleOutputX % (BlendDataTarget->PixelWidth)) + (BlendDataTarget->PixelWidth)) % (BlendDataTarget->PixelWidth); // Positive Mod
+					int32 SampleOutputY = Y;
+					SampleOutputY -= BlendDataTarget->OutputBoundsMin.Y;
 						
-						const int32 FinalIndex = SampleOutputX + (SampleOutputY * (BlendDataTarget->PixelWidth + 1));
-						BlendDataTarget->Data[FinalIndex] += SampleColor * SampleWeightSquared;
-						// OutputEquirectangularMap[OutputPixelX + (OutputPixelY * OutputEquirectangularMapSize.X)] += SampleColor * SampleWeightSquared; // ToDo move to weight map.
-					}
+					const int32 FinalIndex = SampleOutputX + (SampleOutputY * (BlendDataTarget->PixelWidth));
+					BlendDataTarget->Data[FinalIndex] += SampleColor * SampleWeightSquared;
+					// OutputEquirectangularMap[OutputPixelX + (OutputPixelY * OutputEquirectangularMapSize.X)] += SampleColor * SampleWeightSquared; // ToDo move to weight map.
 				}
 			}
 		}
-// #if ACCUMULATE_IN_SCANLINES
-// 	});
-// #endif
+	}
 
 	BlendDataTarget->BlendEndTime = FPlatformTime::Seconds();
 
@@ -330,16 +307,16 @@ void FMoviePipelinePanoramicBlender::OnCompleteRenderPassDataAvailable_AnyThread
 			EyeOffset = (OutputEquirectangularMapSize.X * OutputEquirectangularMapSize.Y) * BlendDataTarget->OriginalDataPayload->Pane.EyeIndex;
 		}
 
-		for (int32 SampleY = 0; SampleY <= BlendDataTarget->PixelHeight; SampleY++)
+		for (int32 SampleY = 0; SampleY < BlendDataTarget->PixelHeight; SampleY++)
 		{
-			for (int32 SampleX = 0; SampleX <= BlendDataTarget->PixelWidth; SampleX++)
+			for (int32 SampleX = 0; SampleX < BlendDataTarget->PixelWidth; SampleX++)
 			{
 				int32 OriginalX = SampleX + BlendDataTarget->OutputBoundsMin.X;
 				int32 OriginalY = SampleY + BlendDataTarget->OutputBoundsMin.Y;
 				const int32 OutputPixelX = ((OriginalX % OutputEquirectangularMapSize.X) + OutputEquirectangularMapSize.X) % OutputEquirectangularMapSize.X;
 				const int32 OutputPixelY = OriginalY;
 
-				int32 SourceIndex = SampleX + (SampleY * (BlendDataTarget->PixelWidth + 1));
+				int32 SourceIndex = SampleX + (SampleY * (BlendDataTarget->PixelWidth));
 				int32 DestIndex = OutputPixelX + (OutputPixelY * OutputEquirectangularMapSize.X);
 				OutputFrame->OutputEquirectangularMap[DestIndex + EyeOffset] += BlendDataTarget->Data[SourceIndex];
 			}
@@ -352,19 +329,19 @@ void FMoviePipelinePanoramicBlender::OnCompleteRenderPassDataAvailable_AnyThread
 			// Hack up the debug output name a bit so they're unique.
 			if (BlendDataTarget->OriginalDataPayload->Pane.EyeIndex >= 0)
 			{
-				BlendDataTarget->OriginalDataPayload->Debug_OverrideFilename = FString::Printf(TEXT("/%s_PaneX_%d_PaneY_%dEye_%d-Blended.%d.exr"),
+				BlendDataTarget->OriginalDataPayload->Debug_OverrideFilename = FString::Printf(TEXT("/%s_PaneX_%d_PaneY_%dEye_%d-Blended.%d"),
 					*BlendDataTarget->OriginalDataPayload->PassIdentifier.Name, BlendDataTarget->OriginalDataPayload->Pane.HorizontalStepIndex,
 					BlendDataTarget->OriginalDataPayload->Pane.VerticalStepIndex, DataPayload->Pane.EyeIndex, BlendDataTarget->OriginalDataPayload->SampleState.OutputState.OutputFrameNumber);
 			}
 			else
 			{
-				BlendDataTarget->OriginalDataPayload->Debug_OverrideFilename = FString::Printf(TEXT("/%s_PaneX_%d_PaneY_%dEye_%d-Blended.%d.exr"),
+				BlendDataTarget->OriginalDataPayload->Debug_OverrideFilename = FString::Printf(TEXT("/%s_PaneX_%d_PaneY_%d-Blended.%d"),
 					*BlendDataTarget->OriginalDataPayload->PassIdentifier.Name, BlendDataTarget->OriginalDataPayload->Pane.HorizontalStepIndex,
 					BlendDataTarget->OriginalDataPayload->Pane.VerticalStepIndex, BlendDataTarget->OriginalDataPayload->SampleState.OutputState.OutputFrameNumber);
 			}
 
 			// Now that the sample has been blended pass it (and the memory it owned, we already read from it) to the debug output step.
-			TUniquePtr<TImagePixelData<FLinearColor>> FinalPixelData = MakeUnique<TImagePixelData<FLinearColor>>(FIntPoint(BlendDataTarget->PixelWidth + 1, BlendDataTarget->PixelHeight + 1), TArray64<FLinearColor>(MoveTemp(BlendDataTarget->Data)), BlendDataTarget->OriginalDataPayload);
+			TUniquePtr<TImagePixelData<FLinearColor>> FinalPixelData = MakeUnique<TImagePixelData<FLinearColor>>(FIntPoint(BlendDataTarget->PixelWidth, BlendDataTarget->PixelHeight), TArray64<FLinearColor>(MoveTemp(BlendDataTarget->Data)), BlendDataTarget->OriginalDataPayload);
 			ensure(OutputMerger.IsValid());
 			OutputMerger.Pin()->OnSingleSampleDataAvailable_AnyThread(MoveTemp(FinalPixelData));
 		}
@@ -407,6 +384,10 @@ void FMoviePipelinePanoramicBlender::OnCompleteRenderPassDataAvailable_AnyThread
 				Pixel.R /= Pixel.A;
 				Pixel.G /= Pixel.A;
 				Pixel.B /= Pixel.A;
+				
+				// Now that we've used the weight stored in the alpha channel, restore it to full opaque-ness so
+				// the resulting png/exrs don't end up semi-transparent.
+				Pixel.A = 1.f;
 			}
 		}
 		TSharedRef<FImagePixelDataPayload, ESPMode::ThreadSafe> NewPayload = DataPayload->Copy();
diff --git a/Engine/Plugins/MovieScene/MoviePipelineMaskRenderPass/Source/MoviePipelineMaskRenderPass/Private/MoviePipelinePanoramicPass.h b/Engine/Plugins/MovieScene/MoviePipelineMaskRenderPass/Source/MoviePipelineMaskRenderPass/Private/MoviePipelinePanoramicPass.h
--- a/Engine/Plugins/MovieScene/MoviePipelineMaskRenderPass/Source/MoviePipelineMaskRenderPass/Private/MoviePipelinePanoramicPass.h
+++ b/Engine/Plugins/MovieScene/MoviePipelineMaskRenderPass/Source/MoviePipelineMaskRenderPass/Private/MoviePipelinePanoramicPass.h
@@ -119,7 +119,7 @@ protected:
 	virtual void GatherOutputPassesImpl(TArray<FMoviePipelinePassIdentifier>& ExpectedRenderPasses) override;
 	virtual bool IsAntiAliasingSupported() const { return true; }
 	virtual int32 GetOutputFileSortingOrder() const override { return 1; }
-	virtual bool IsAlphaInTonemapperRequiredImpl() const override { return true; }
+	virtual bool IsAlphaInTonemapperRequiredImpl() const override { return false; }
 	virtual FSceneViewStateInterface* GetSceneViewStateInterface(IViewCalcPayload* OptPayload) override;
 	virtual UTextureRenderTarget2D* GetViewRenderTarget(IViewCalcPayload* OptPayload) const override;
 	virtual void AddViewExtensions(FSceneViewFamilyContext& InContext, FMoviePipelineRenderPassMetrics& InOutSampleState) override;
