MoviePipeline: Added an option to the ObjectID pass to decide whether or not to include transparent objects should be included. If included, they will write to the ObjectID pass as opaque pixels.

diff --git a/Engine/Plugins/MovieScene/MoviePipelineMaskRenderPass/Source/MoviePipelineMaskRenderPass/MoviePipelineMaskRenderPass.Build.cs b/Engine/Plugins/MovieScene/MoviePipelineMaskRenderPass/Source/MoviePipelineMaskRenderPass/MoviePipelineMaskRenderPass.Build.cs
--- a/Engine/Plugins/MovieScene/MoviePipelineMaskRenderPass/Source/MoviePipelineMaskRenderPass/MoviePipelineMaskRenderPass.Build.cs
+++ b/Engine/Plugins/MovieScene/MoviePipelineMaskRenderPass/Source/MoviePipelineMaskRenderPass/MoviePipelineMaskRenderPass.Build.cs
@@ -25,5 +25,11 @@ public class MoviePipelineMaskRenderPass : ModuleRules
 				"OpenColorIO",
 			}
 		);
+
+		if (Target.bBuildEditor)
+		{
+			PublicDependencyModuleNames.AddRange(
+			new string[] { "UnrealEd" });
+		}
 	}
 }
diff --git a/Engine/Plugins/MovieScene/MoviePipelineMaskRenderPass/Source/MoviePipelineMaskRenderPass/Private/MoviePipelineObjectIdPass.cpp b/Engine/Plugins/MovieScene/MoviePipelineMaskRenderPass/Source/MoviePipelineMaskRenderPass/Private/MoviePipelineObjectIdPass.cpp
--- a/Engine/Plugins/MovieScene/MoviePipelineMaskRenderPass/Source/MoviePipelineMaskRenderPass/Private/MoviePipelineObjectIdPass.cpp
+++ b/Engine/Plugins/MovieScene/MoviePipelineMaskRenderPass/Source/MoviePipelineMaskRenderPass/Private/MoviePipelineObjectIdPass.cpp
@@ -25,6 +25,10 @@
 #include "UObject/Package.h"
 #include "UObject/UObjectAnnotation.h"
 
+#if WITH_EDITOR
+#include "Editor/EditorPerProjectUserSettings.h"
+#endif
+
 #include UE_INLINE_GENERATED_CPP_BY_NAME(MoviePipelineObjectIdPass)
 
 DECLARE_CYCLE_STAT(TEXT("STAT_MoviePipeline_AccumulateMaskSample_TT"), STAT_AccumulateMaskSample_TaskThread, STATGROUP_MoviePipeline);
@@ -182,12 +186,23 @@ void UMoviePipelineObjectIdRenderPass::SetupImpl(const MoviePipeline::FMoviePipe
 		AccelData.JsonManifest->SetStringField(TEXT("default"), FString::Printf(TEXT("%08x"), DefaultHash));
 	}
 	ManifestAnnotation.AddAnnotation(this, AccelData);
+
+#if WITH_EDITOR
+	UEditorPerProjectUserSettings* EditorSettings = GetMutableDefault<UEditorPerProjectUserSettings>();
+	bPrevAllowSelectTranslucent = EditorSettings->bAllowSelectTranslucent;
+	EditorSettings->bAllowSelectTranslucent = bIncludeTranslucentObjects;
+#endif
 }
 
 void UMoviePipelineObjectIdRenderPass::TeardownImpl()
 {
 	ManifestAnnotation.RemoveAnnotation(this);
 
+#if WITH_EDITOR
+	UEditorPerProjectUserSettings* EditorSettings = GetMutableDefault<UEditorPerProjectUserSettings>();
+	EditorSettings->bAllowSelectTranslucent = bPrevAllowSelectTranslucent;
+#endif
+
 	// Preserve our view state until the rendering thread has been flushed.
 	Super::TeardownImpl();
 }
diff --git a/Engine/Plugins/MovieScene/MoviePipelineMaskRenderPass/Source/MoviePipelineMaskRenderPass/Private/MoviePipelineObjectIdPass.h b/Engine/Plugins/MovieScene/MoviePipelineMaskRenderPass/Source/MoviePipelineMaskRenderPass/Private/MoviePipelineObjectIdPass.h
--- a/Engine/Plugins/MovieScene/MoviePipelineMaskRenderPass/Source/MoviePipelineMaskRenderPass/Private/MoviePipelineObjectIdPass.h
+++ b/Engine/Plugins/MovieScene/MoviePipelineMaskRenderPass/Source/MoviePipelineMaskRenderPass/Private/MoviePipelineObjectIdPass.h
@@ -50,7 +50,12 @@ public:
 	UPROPERTY(EditAnywhere, BlueprintReadWrite, Category = "Settings")
 	EMoviePipelineObjectIdPassIdType IdType;
 
+	/** If true, translucent objects will be included in the ObjectId pass (but as an opaque layer due to limitations). False will omit translucent objects. */
+	UPROPERTY(EditAnywhere, BlueprintReadWrite, Category = "Settings")
+	bool bIncludeTranslucentObjects;
+
 private:
 	TSharedPtr<FAccumulatorPool, ESPMode::ThreadSafe> AccumulatorPool;
 	TArray<FMoviePipelinePassIdentifier> ExpectedPassIdentifiers;
+	bool bPrevAllowSelectTranslucent;
 };
\ No newline at end of file
