MovieRenderQueue: fixed Panoramic rendering, implemnted missed base class functions required for each pane render target size calculation

diff --git a/Engine/Plugins/MovieScene/MoviePipelineMaskRenderPass/Source/MoviePipelineMaskRenderPass/Private/MoviePipelinePanoramicPass.cpp b/Engine/Plugins/MovieScene/MoviePipelineMaskRenderPass/Source/MoviePipelineMaskRenderPass/Private/MoviePipelinePanoramicPass.cpp
--- a/Engine/Plugins/MovieScene/MoviePipelineMaskRenderPass/Source/MoviePipelineMaskRenderPass/Private/MoviePipelinePanoramicPass.cpp
+++ b/Engine/Plugins/MovieScene/MoviePipelineMaskRenderPass/Source/MoviePipelineMaskRenderPass/Private/MoviePipelinePanoramicPass.cpp
@@ -381,6 +381,27 @@ FSceneView* UMoviePipelinePanoramicPass::GetSceneViewForSampleState(FSceneViewFa
 	return View;
 }
 
+FIntPoint UMoviePipelinePanoramicPass::GetPayloadPaneResolution(const FIntPoint& InSize, IViewCalcPayload* OptPayload) const
+{
+	if (OptPayload)
+	{
+		FPanoPane* PanoPane = (FPanoPane*)OptPayload;
+		return PanoPane->Resolution;
+	}
+
+	return InSize;
+}
+
+TWeakObjectPtr<UTextureRenderTarget2D> UMoviePipelinePanoramicPass::GetOrCreateViewRenderTarget(const FIntPoint& InSize, IViewCalcPayload* OptPayload)
+{
+	return Super::GetOrCreateViewRenderTarget(GetPayloadPaneResolution(InSize, OptPayload), OptPayload);
+}
+
+TSharedPtr<FMoviePipelineSurfaceQueue, ESPMode::ThreadSafe> UMoviePipelinePanoramicPass::GetOrCreateSurfaceQueue(const FIntPoint& InSize, IViewCalcPayload* OptPayload)
+{
+	return Super::GetOrCreateSurfaceQueue(GetPayloadPaneResolution(InSize, OptPayload), OptPayload);
+}
+
 void UMoviePipelinePanoramicPass::RenderSample_GameThreadImpl(const FMoviePipelineRenderPassMetrics& InSampleState)
 {
 	// Wait for a surface to be available to write to. This will stall the game thread while the RHI/Render Thread catch up.
@@ -468,7 +489,6 @@ void UMoviePipelinePanoramicPass::RenderSample_GameThreadImpl(const FMoviePipeli
 	}
 }
 
-
 void UMoviePipelinePanoramicPass::ScheduleReadbackAndAccumulation(const FMoviePipelineRenderPassMetrics& InSampleState, const FPanoPane& InPane, FCanvas& InCanvas)
 {
 	// If this was just to contribute to the history buffer, no need to go any further.
@@ -496,6 +516,7 @@ void UMoviePipelinePanoramicPass::ScheduleReadbackAndAccumulation(const FMoviePi
 	FramePayload->SampleState = InSampleState;
 	FramePayload->SortingOrder = GetOutputFileSortingOrder();
 	FramePayload->Pane = InPane;
+
 	if (FramePayload->Pane.EyeIndex >= 0)
 	{
 		FramePayload->Debug_OverrideFilename = FString::Printf(TEXT("/%s_SS_%d_TS_%d_TileX_%d_TileY_%d_PaneX_%d_PaneY_%d_Eye_%d.%d.exr"),
@@ -511,8 +532,7 @@ void UMoviePipelinePanoramicPass::ScheduleReadbackAndAccumulation(const FMoviePi
 			FramePayload->Pane.VerticalStepIndex, FramePayload->SampleState.OutputState.OutputFrameNumber);
 	}
 	
-	const FIntPoint PaneResolution = GetPaneResolution(InSampleState.BackbufferSize);
-	TSharedPtr<FMoviePipelineSurfaceQueue, ESPMode::ThreadSafe> LocalSurfaceQueue = GetOrCreateSurfaceQueue(PaneResolution, (IViewCalcPayload*)(&FramePayload.Get()));
+	TSharedPtr<FMoviePipelineSurfaceQueue, ESPMode::ThreadSafe> LocalSurfaceQueue = GetOrCreateSurfaceQueue(InSampleState.BackbufferSize, (IViewCalcPayload*)(&FramePayload->Pane));
 
 	MoviePipeline::FImageSampleAccumulationArgs AccumulationArgs;
 	{
@@ -555,5 +575,3 @@ void UMoviePipelinePanoramicPass::ScheduleReadbackAndAccumulation(const FMoviePi
 		});
 
 }
-
-
diff --git a/Engine/Plugins/MovieScene/MoviePipelineMaskRenderPass/Source/MoviePipelineMaskRenderPass/Private/MoviePipelinePanoramicPass.h b/Engine/Plugins/MovieScene/MoviePipelineMaskRenderPass/Source/MoviePipelineMaskRenderPass/Private/MoviePipelinePanoramicPass.h
--- a/Engine/Plugins/MovieScene/MoviePipelineMaskRenderPass/Source/MoviePipelineMaskRenderPass/Private/MoviePipelinePanoramicPass.h
+++ b/Engine/Plugins/MovieScene/MoviePipelineMaskRenderPass/Source/MoviePipelineMaskRenderPass/Private/MoviePipelinePanoramicPass.h
@@ -124,6 +124,8 @@ protected:
 	virtual void AddViewExtensions(FSceneViewFamilyContext& InContext, FMoviePipelineRenderPassMetrics& InOutSampleState) override;
 	virtual bool IsAutoExposureAllowed(const FMoviePipelineRenderPassMetrics& InSampleState) const override { return false; }
 	virtual FSceneView* GetSceneViewForSampleState(FSceneViewFamily* ViewFamily, FMoviePipelineRenderPassMetrics& InOutSampleState, IViewCalcPayload* OptPayload = nullptr) override;
+	virtual TWeakObjectPtr<UTextureRenderTarget2D> GetOrCreateViewRenderTarget(const FIntPoint& InSize, IViewCalcPayload* OptPayload = nullptr) override;
+	virtual TSharedPtr<FMoviePipelineSurfaceQueue, ESPMode::ThreadSafe> GetOrCreateSurfaceQueue(const FIntPoint& InSize, IViewCalcPayload* OptPayload = nullptr) override;
 	// ~UMoviePipelineRenderPass
 
 	// FGCObject Interface
@@ -133,7 +135,7 @@ protected:
 	void ScheduleReadbackAndAccumulation(const FMoviePipelineRenderPassMetrics& InSampleState, const FPanoPane& InPane, FCanvas& InCanvas);
 	void GetFieldOfView(float& OutHorizontal, float& OutVertical, const bool bInStereo) const;
 	FIntPoint GetPaneResolution(const FIntPoint& InSize) const;
-
+	FIntPoint GetPayloadPaneResolution(const FIntPoint& InSize, IViewCalcPayload* OptPayload) const;
 public:
 
 	/**
