MoviePipeline: Added support for camera rotation to be preserved during a Panoramic Render.

diff --git a/Engine/Plugins/MovieScene/MoviePipelineMaskRenderPass/Source/MoviePipelineMaskRenderPass/Private/MoviePipelinePanoramicBlender.cpp b/Engine/Plugins/MovieScene/MoviePipelineMaskRenderPass/Source/MoviePipelineMaskRenderPass/Private/MoviePipelinePanoramicBlender.cpp
--- a/Engine/Plugins/MovieScene/MoviePipelineMaskRenderPass/Source/MoviePipelineMaskRenderPass/Private/MoviePipelinePanoramicBlender.cpp
+++ b/Engine/Plugins/MovieScene/MoviePipelineMaskRenderPass/Source/MoviePipelineMaskRenderPass/Private/MoviePipelinePanoramicBlender.cpp
@@ -110,7 +110,7 @@ void FMoviePipelinePanoramicBlender::OnCompleteRenderPassDataAvailable_AnyThread
 	// We calculate all of these immediately on the task thread because it's read-only and we need the
 	// information calculated by it to allocate memory when we start working on a sample.
 	FIntPoint SampleSize = DataPayload->Pane.Resolution;
-	FRotator SampleRotation = FRotator(DataPayload->Pane.CameraRotation);
+	FRotator SampleRotation = FRotator(DataPayload->Pane.bUseLocalRotation ? DataPayload->Pane.CameraLocalRotation : DataPayload->Pane.CameraRotation);
 	const float SampleHalfHorizontalFoVDegrees = 0.5f * DataPayload->Pane.HorizontalFieldOfView;
 	const float SampleHalfVerticalFoVDegrees = 0.5f * DataPayload->Pane.VerticalFieldOfView; // FMath::RadiansToDegrees((FMath::Atan((0.5f * SampleSize.X / FMath::Tan(DataPayload->Pane.HorizontalFieldOfView / 2.f))))); // CTransformTools::getVFovFromHFovAndARInDegrees(iSampleFovInDegrees, static_cast<float>(iSampleSize.X) / static_cast<float>(iSampleSize.Y));
 	const float SampleHalfHorizontalFoVCosine = FMath::Cos(FMath::DegreesToRadians(SampleHalfHorizontalFoVDegrees));
diff --git a/Engine/Plugins/MovieScene/MoviePipelineMaskRenderPass/Source/MoviePipelineMaskRenderPass/Private/MoviePipelinePanoramicPass.cpp b/Engine/Plugins/MovieScene/MoviePipelineMaskRenderPass/Source/MoviePipelineMaskRenderPass/Private/MoviePipelinePanoramicPass.cpp
--- a/Engine/Plugins/MovieScene/MoviePipelineMaskRenderPass/Source/MoviePipelineMaskRenderPass/Private/MoviePipelinePanoramicPass.cpp
+++ b/Engine/Plugins/MovieScene/MoviePipelineMaskRenderPass/Source/MoviePipelineMaskRenderPass/Private/MoviePipelinePanoramicPass.cpp
@@ -22,6 +22,7 @@ UMoviePipelinePanoramicPass::UMoviePipelinePanoramicPass()
 	: UMoviePipelineImagePassBase()
 	, NumHorizontalSteps(8)
 	, NumVerticalSteps(3)
+	, bFollowCameraOrientation(true)
 	, bStereo(false)
 	, EyeSeparation(6.5f)
 	, EyeConvergenceDistance(EyeSeparation * 30.f)
@@ -52,7 +53,7 @@ namespace MoviePipeline
 			return Results;
 		};
 
-		void GetCameraOrientationForStereo(FVector& OutLocation, FRotator& OutRotation, const FPanoPane& InPane, const int32 InStereoIndex, const bool bInPrevPosition)
+		void GetCameraOrientationForStereo(FVector& OutLocation, FRotator& OutRotation, FRotator& OutLocalRotation, const FPanoPane& InPane, const int32 InStereoIndex, const bool bInPrevPosition)
 		{
 			// ToDo: This 110 (-55, 55) comes from TwinMotion who uses a hard-coded number of v-steps, may need adjusting.
 			const TArray<float> PitchValues = MoviePipeline::Panoramic::DistributeValuesInInterval(-55, 55, InPane.NumVerticalSteps, /*Inclusive Max*/true);
@@ -67,6 +68,7 @@ namespace MoviePipeline
 			const FRotator SourceRot = bInPrevPosition ? InPane.PrevOriginalCameraRotation : InPane.OriginalCameraRotation;
 			FQuat RotationResult = FQuat(SourceRot) * HorizontalRotQuat * VerticalRotQuat;
 			OutRotation = FRotator(RotationResult);
+			OutLocalRotation = FRotator(HorizontalRotQuat * VerticalRotQuat);
 
 			// If not using stereo rendering then the eye is just the camera location
 			if (InStereoIndex < 0)
@@ -420,10 +422,12 @@ void UMoviePipelinePanoramicPass::RenderSample_GameThreadImpl(const FMoviePipeli
 				Pane.NumVerticalSteps = NumVerticalSteps;
 				Pane.EyeSeparation = EyeSeparation;
 				Pane.EyeConvergenceDistance = EyeConvergenceDistance;
+				Pane.bUseLocalRotation = bFollowCameraOrientation;
 
 				// Get the actual camera location/rotation for this particular pane, the above values are from the global camera.
-				MoviePipeline::Panoramic::GetCameraOrientationForStereo(/*Out*/ Pane.CameraLocation, /*Out*/ Pane.CameraRotation, Pane, StereoIndex, /*bInPrevPos*/ false);
-				MoviePipeline::Panoramic::GetCameraOrientationForStereo(/*Out*/ Pane.PrevCameraLocation, /*Out*/ Pane.PrevCameraRotation, Pane, StereoIndex, /*bInPrevPos*/ true);
+				MoviePipeline::Panoramic::GetCameraOrientationForStereo(/*Out*/ Pane.CameraLocation, /*Out*/ Pane.CameraRotation, /*Out*/ Pane.CameraLocalRotation, Pane, StereoIndex, /*bInPrevPos*/ false);
+				FRotator DummyPrevLocalRot;
+				MoviePipeline::Panoramic::GetCameraOrientationForStereo(/*Out*/ Pane.PrevCameraLocation, /*Out*/ Pane.PrevCameraRotation, /*Out*/ DummyPrevLocalRot, Pane, StereoIndex, /*bInPrevPos*/ true);
 
 				GetFieldOfView(Pane.HorizontalFieldOfView, Pane.VerticalFieldOfView, bStereo);
 
diff --git a/Engine/Plugins/MovieScene/MoviePipelineMaskRenderPass/Source/MoviePipelineMaskRenderPass/Private/MoviePipelinePanoramicPass.h b/Engine/Plugins/MovieScene/MoviePipelineMaskRenderPass/Source/MoviePipelineMaskRenderPass/Private/MoviePipelinePanoramicPass.h
--- a/Engine/Plugins/MovieScene/MoviePipelineMaskRenderPass/Source/MoviePipelineMaskRenderPass/Private/MoviePipelinePanoramicPass.h
+++ b/Engine/Plugins/MovieScene/MoviePipelineMaskRenderPass/Source/MoviePipelineMaskRenderPass/Private/MoviePipelinePanoramicPass.h
@@ -40,8 +40,12 @@ struct FPanoPane : public UMoviePipelineImagePassBase::IViewCalcPayload
 	FVector CameraLocation;
 	FVector PrevCameraLocation;
 	FRotator CameraRotation;
+	FRotator CameraLocalRotation;
 	FRotator PrevCameraRotation;
 
+	// If true, uses only the CameraLocalRotation which means that if the camera yaws, so will the resulting blended image.
+	bool bUseLocalRotation;
+
 	// How many horizontal segments are there total.
 	int32 NumHorizontalSteps;
 	int32 NumVerticalSteps;
@@ -152,6 +156,12 @@ public:
 	*/
 	UPROPERTY(EditAnywhere, BlueprintReadWrite, Category = "Panoramic Settings", meta = (UIMin = "1", ClampMin = "1"))
 	int32 NumVerticalSteps;
+
+	/**
+	* Should the Pitch, Yaw and Roll of the camera be respected? If false, only the location will be taken from the camera.
+	*/
+	UPROPERTY(EditAnywhere, BlueprintReadWrite, Category = "Panoramic Settings")
+	bool bFollowCameraOrientation;
 	
 	/**
 	* If true we will capture a stereo panorama. This doubles everything (all render times, memory requirements, etc.)
