MoviePipeline: Added support for the Object ID pass (Cryptomatte) to the graph via a new "Object ID" node.

Refactored some of the older MRQ Object ID code into shared utilities so the common logic can be shared with the new MRG Object ID node.

diff --git a/Engine/Plugins/MovieScene/MoviePipelineMaskRenderPass/Source/MoviePipelineMaskRenderPass/MoviePipelineMaskRenderPass.Build.cs b/Engine/Plugins/MovieScene/MoviePipelineMaskRenderPass/Source/MoviePipelineMaskRenderPass/MoviePipelineMaskRenderPass.Build.cs
--- a/Engine/Plugins/MovieScene/MoviePipelineMaskRenderPass/Source/MoviePipelineMaskRenderPass/MoviePipelineMaskRenderPass.Build.cs
+++ b/Engine/Plugins/MovieScene/MoviePipelineMaskRenderPass/Source/MoviePipelineMaskRenderPass/MoviePipelineMaskRenderPass.Build.cs
@@ -11,7 +11,8 @@ public class MoviePipelineMaskRenderPass : ModuleRules
 				"Core",
 				"CoreUObject",
 				"Engine",
-				"Json"
+				"Json",
+				"ImageCore",
 			}
 		);
 
@@ -23,6 +24,7 @@ public class MoviePipelineMaskRenderPass : ModuleRules
                 "RHI",
 				"ActorLayerUtilities",
 				"OpenColorIO",
+				"SlateCore",
 			}
 		);
 
diff --git a/Engine/Plugins/MovieScene/MoviePipelineMaskRenderPass/Source/MoviePipelineMaskRenderPass/Private/Graph/MovieGraphObjectIdNode.cpp b/Engine/Plugins/MovieScene/MoviePipelineMaskRenderPass/Source/MoviePipelineMaskRenderPass/Private/Graph/MovieGraphObjectIdNode.cpp
new file mode 100644
--- /dev/null
+++ b/Engine/Plugins/MovieScene/MoviePipelineMaskRenderPass/Source/MoviePipelineMaskRenderPass/Private/Graph/MovieGraphObjectIdNode.cpp
@@ -0,0 +1,86 @@
+﻿// Copyright Epic Games, Inc. All Rights Reserved.
+
+#include "MovieGraphObjectIdNode.h"
+
+#include "Editor/EditorPerProjectUserSettings.h"
+#include "MovieGraphObjectIdPass.h"
+#include "MoviePipelineObjectIdUtils.h"
+
+FUObjectAnnotationSparse<UE::MoviePipeline::FObjectIdAccelerationData, true>& UMovieGraphObjectIdNode::GetManifestAnnotation()
+{
+	static FUObjectAnnotationSparse<UE::MoviePipeline::FObjectIdAccelerationData, true> ManifestAnnotation;
+	
+	return ManifestAnnotation;
+}
+
+FEngineShowFlags UMovieGraphObjectIdNode::GetShowFlags() const
+{
+	FEngineShowFlags Flags(ESFIM_Game);
+	Flags.DisableAdvancedFeatures();
+	Flags.SetPostProcessing(false);
+	Flags.SetPostProcessMaterial(false);
+
+	// The most important flag. The Hit Proxy IDs will be used to generate Object IDs.
+	Flags.SetHitProxies(true);
+
+	// Screen-percentage scaling mixes IDs when doing down-sampling, so it is disabled.
+	Flags.SetScreenPercentage(false);
+
+	return Flags;
+}
+
+EViewModeIndex UMovieGraphObjectIdNode::GetViewModeIndex() const
+{
+	return VMI_Unlit;
+}
+
+bool UMovieGraphObjectIdNode::GetAllowsShowFlagsCustomization() const
+{
+	return false;
+}
+
+#if WITH_EDITOR
+FText UMovieGraphObjectIdNode::GetNodeTitle(const bool bGetDescriptive) const
+{
+	return NSLOCTEXT("MovieRenderGraph", "ObjectIdNodeTitle", "Object IDs");		
+}
+
+FSlateIcon UMovieGraphObjectIdNode::GetIconAndTint(FLinearColor& OutColor) const
+{
+	static const FSlateIcon DeferredRendererIcon = FSlateIcon(FAppStyle::GetAppStyleSetName(), "ContentBrowser.SizeMap");
+	OutColor = FLinearColor::White;
+	
+	return DeferredRendererIcon;
+}
+#endif	// WITH_EDITOR
+
+TUniquePtr<UE::MovieGraph::Rendering::FMovieGraphImagePassBase> UMovieGraphObjectIdNode::CreateInstance() const
+{
+	return MakeUnique<FMovieGraphObjectIdPass>();
+}
+
+FString UMovieGraphObjectIdNode::GetRendererNameImpl() const
+{
+	return TEXT("ObjectID");
+}
+
+void UMovieGraphObjectIdNode::SetupImpl(const FMovieGraphRenderPassSetupData& InSetupData)
+{
+	Super::SetupImpl(InSetupData);
+
+#if WITH_EDITOR
+	UEditorPerProjectUserSettings* EditorSettings = GetMutableDefault<UEditorPerProjectUserSettings>();
+	bPrevAllowSelectTranslucent = EditorSettings->bAllowSelectTranslucent;
+	EditorSettings->bAllowSelectTranslucent = bIncludeTranslucentObjects;
+#endif
+}
+
+void UMovieGraphObjectIdNode::TeardownImpl()
+{
+#if WITH_EDITOR
+	UEditorPerProjectUserSettings* EditorSettings = GetMutableDefault<UEditorPerProjectUserSettings>();
+	EditorSettings->bAllowSelectTranslucent = bPrevAllowSelectTranslucent;
+#endif
+	
+	Super::TeardownImpl();
+}
diff --git a/Engine/Plugins/MovieScene/MoviePipelineMaskRenderPass/Source/MoviePipelineMaskRenderPass/Private/Graph/MovieGraphObjectIdNode.h b/Engine/Plugins/MovieScene/MoviePipelineMaskRenderPass/Source/MoviePipelineMaskRenderPass/Private/Graph/MovieGraphObjectIdNode.h
new file mode 100644
--- /dev/null
+++ b/Engine/Plugins/MovieScene/MoviePipelineMaskRenderPass/Source/MoviePipelineMaskRenderPass/Private/Graph/MovieGraphObjectIdNode.h
@@ -0,0 +1,69 @@
+// Copyright Epic Games, Inc. All Rights Reserved.
+
+#pragma once
+
+#include "MoviePipelineObjectIdUtils.h"
+#include "Graph/Nodes/MovieGraphImagePassBaseNode.h"
+#include "UObject/UObjectAnnotation.h"
+
+#include "MovieGraphObjectIdNode.generated.h"
+
+enum class EMoviePipelineObjectIdPassIdType : uint8;
+
+/** A node which outputs layers compatible with the Cryptomatte standard. */
+UCLASS(BlueprintType)
+class UMovieGraphObjectIdNode : public UMovieGraphImagePassBaseNode
+{
+	GENERATED_BODY()
+
+public:
+	UMovieGraphObjectIdNode() = default;
+
+	static FUObjectAnnotationSparse<UE::MoviePipeline::FObjectIdAccelerationData, true>& GetManifestAnnotation();
+
+	// UMovieGraphImagePassBaseNode Interface
+	virtual FEngineShowFlags GetShowFlags() const override;
+	virtual EViewModeIndex GetViewModeIndex() const override;
+	virtual bool GetAllowsShowFlagsCustomization() const override;
+	// ~UMovieGraphImagePassBaseNode Interface
+
+	// UMovieGraphNode Interface
+#if WITH_EDITOR
+	virtual FText GetNodeTitle(const bool bGetDescriptive = false) const override;
+	virtual FSlateIcon GetIconAndTint(FLinearColor& OutColor) const override;
+#endif
+	// ~UMovieGraphNode Interface
+
+protected:
+	// UMovieGraphImagePassBaseNode Interface
+	virtual TUniquePtr<UE::MovieGraph::Rendering::FMovieGraphImagePassBase> CreateInstance() const override;
+	// ~UMovieGraphImagePassBaseNode Interface
+
+	// UMovieGraphRenderPassNode Interface
+	virtual FString GetRendererNameImpl() const override;
+	virtual void SetupImpl(const FMovieGraphRenderPassSetupData& InSetupData) override;
+	virtual void TeardownImpl() override;
+	// ~UMovieGraphRenderPassNode Interface
+
+public:
+	UPROPERTY(EditAnywhere, BlueprintReadWrite, Category = Overrides, meta = (InlineEditConditionToggle))
+	uint8 bOverride_IdType : 1;
+
+	UPROPERTY(EditAnywhere, BlueprintReadWrite, Category = Overrides, meta = (InlineEditConditionToggle))
+	uint8 bOverride_bIncludeTranslucentObjects : 1;
+
+	/** How objects are grouped together within the Cryptomatte. */
+	UPROPERTY(EditAnywhere, BlueprintReadWrite, Category = "Settings", DisplayName = "ID Type", meta = (EditCondition="bOverride_IdType"))
+	EMoviePipelineObjectIdPassIdType IdType;
+
+	/** If true, translucent objects will be included in the ObjectId pass (but as an opaque layer due to limitations). False will omit translucent objects. */
+	UPROPERTY(EditAnywhere, BlueprintReadWrite, Category = "Settings", meta = (EditCondition="bOverride_bIncludeTranslucentObjects"))
+	bool bIncludeTranslucentObjects;
+
+private:
+	/**
+	 * The prior value of the project's bAllowSelectTranslucent setting. This node will temporarily change the value of this property to the value
+	 * set in bIncludeTranslucentObjects while rendering.
+	 */
+	bool bPrevAllowSelectTranslucent;
+};
\ No newline at end of file
diff --git a/Engine/Plugins/MovieScene/MoviePipelineMaskRenderPass/Source/MoviePipelineMaskRenderPass/Private/Graph/MovieGraphObjectIdPass.cpp b/Engine/Plugins/MovieScene/MoviePipelineMaskRenderPass/Source/MoviePipelineMaskRenderPass/Private/Graph/MovieGraphObjectIdPass.cpp
new file mode 100644
--- /dev/null
+++ b/Engine/Plugins/MovieScene/MoviePipelineMaskRenderPass/Source/MoviePipelineMaskRenderPass/Private/Graph/MovieGraphObjectIdPass.cpp
@@ -0,0 +1,267 @@
+// Copyright Epic Games, Inc. All Rights Reserved.
+
+#include "MovieGraphObjectIdPass.h"
+
+#include "MovieGraphObjectIdNode.h"
+#include "MoviePipelineHashUtils.h"
+#include "MoviePipelineObjectIdUtils.h"
+#include "Graph/MovieGraphBlueprintLibrary.h"
+#include "Graph/MovieGraphPipeline.h"
+#include "Graph/Nodes/MovieGraphImagePassBaseNode.h"
+
+namespace UE::MovieGraph
+{
+	static void AccumulateSampleObjectId_TaskThread(TUniquePtr<FImagePixelData>&& InPixelData, const FMovieGraphSampleState InSampleState, const TSharedRef<::MoviePipeline::IMoviePipelineAccumulationArgs> InAccumulatorArgs)
+	{
+		TRACE_CPUPROFILER_EVENT_SCOPE(AccumulateSampleObjectId_TaskThread);
+
+		TUniquePtr<FImagePixelData> SamplePixelData = MoveTemp(InPixelData);
+
+		// Associate the sample state with the image as payload data, this allows downstream systems to fetch the values without us having to store the data
+		// separately and ensure they stay paired the whole way down.
+		TSharedPtr<FMovieGraphSampleState> SampleStatePayload = MakeShared<FMovieGraphSampleState>(InSampleState);
+		SamplePixelData->SetPayload(StaticCastSharedPtr<IImagePixelDataPayload>(SampleStatePayload));
+		
+		const TSharedRef<FMovieGraphObjectIdMaskSampleAccumulationArgs> ObjectIdArgs = StaticCastSharedRef<FMovieGraphObjectIdMaskSampleAccumulationArgs>(InAccumulatorArgs);
+
+		const TSharedPtr<IMovieGraphOutputMerger, ESPMode::ThreadSafe> OutputMergerPin = ObjectIdArgs->OutputMerger;
+		if (!OutputMergerPin.IsValid())
+		{
+			return;
+		}
+
+		const bool bIsWellFormed = SamplePixelData->IsDataWellFormed();
+		check(bIsWellFormed);
+
+		if (SampleStatePayload->bWriteSampleToDisk)
+		{
+			// Debug Feature: Write the raw sample to disk for debugging purposes. We copy the data here,
+			// as we don't want to disturb the memory flow below.
+			TUniquePtr<FImagePixelData> SampleData = SamplePixelData->CopyImageData();
+			OutputMergerPin->OnSingleSampleDataAvailable_AnyThread(MoveTemp(SampleData));
+		}
+
+		const TSharedPtr<FMaskOverlappedAccumulator, ESPMode::ThreadSafe> AccumulatorPin = StaticCastWeakPtr<FMaskOverlappedAccumulator>(ObjectIdArgs->ImageAccumulator).Pin();
+
+		// TODO: Need a way to detect if this is the first tile *in addition to* the first temporal sample
+		// For the first sample in a new output, we allocate memory
+		if (SampleStatePayload->TraversalContext.Time.bIsFirstTemporalSampleForFrame)
+		{
+			LLM_SCOPE_BYNAME(TEXT("MoviePipeline/ImageAccumulatorInitMemory"));
+			
+			const FIntPoint PlaneSize = FIntPoint(SampleStatePayload->AccumulatorResolution);
+			
+			AccumulatorPin->InitMemory(PlaneSize);
+			AccumulatorPin->ZeroPlanes();
+		}
+
+		// Accumulate the new sample to our target
+		{
+			TRACE_CPUPROFILER_EVENT_SCOPE(MoviePipeline_AccumulatePixelData);
+
+			FIntPoint RawSize = SamplePixelData->GetSize();
+
+			const void* RawData;
+			int64 TotalSize;
+			SamplePixelData->GetRawData(RawData, TotalSize);
+
+			const FColor* RawDataPtr = static_cast<const FColor*>(RawData);
+			TArray64<float> IdData;
+			IdData.SetNumUninitialized(RawSize.X * RawSize.Y);
+
+			// Remap HitProxy ID into precalculated Cryptomatte hash
+			RemapHitProxyIdToCryptomatteHash(RawSize, RawDataPtr, ObjectIdArgs->CacheData, IdData);
+
+			const FIntPoint TileSize = SampleStatePayload->BackbufferResolution;
+			const FIntPoint OverlappedPad = SampleStatePayload->OverlappedPad;
+			const FIntPoint OverlappedOffset = SampleStatePayload->OverlappedOffset;
+			const FVector2D OverlappedSubpixelShift = SampleStatePayload->OverlappedSubpixelShift;
+			::MoviePipeline::FTileWeight1D WeightFunctionX;
+			::MoviePipeline::FTileWeight1D WeightFunctionY;
+			WeightFunctionX.InitHelper(OverlappedPad.X, TileSize.X, OverlappedPad.X);
+			WeightFunctionY.InitHelper(OverlappedPad.Y, TileSize.Y, OverlappedPad.Y);
+			
+			AccumulatorPin->AccumulatePixelData(IdData.GetData(), RawSize, OverlappedOffset, OverlappedSubpixelShift, WeightFunctionX, WeightFunctionY);
+		}
+
+		// Finally on our last sample, we fetch the data out of the accumulator
+		// and move it to the Output Merger.
+		if (SampleStatePayload->bFetchFromAccumulator)
+		{
+			const int32 FullSizeX = AccumulatorPin->PlaneSize.X;
+			const int32 FullSizeY = AccumulatorPin->PlaneSize.Y;
+
+			// Now that a tile is fully built and accumulated we can notify the output builder that the
+			// data is ready so it can pass that onto the output containers (if needed).
+			// 32 bit FLinearColor
+			TArray<TArray64<FLinearColor>> OutputLayers;
+			for (int32 Index = 0; Index < ObjectIdArgs->NumOutputLayers; Index++)
+			{
+				OutputLayers.Add(TArray64<FLinearColor>());
+			}
+
+			AccumulatorPin->FetchFinalPixelDataLinearColor(OutputLayers);
+
+			// Add in the object ID metadata. This cannot be done in the node's GetFormatResolveArgs() because the manifest is only known after render-time,
+			// and the manifest data is destroyed during node teardown (and teardown occurs before the metadata is finalized and the file written to dis
+			FUObjectAnnotationSparse<UE::MoviePipeline::FObjectIdAccelerationData, true>& ManifestAnnotation = UMovieGraphObjectIdNode::GetManifestAnnotation();
+			const UE::MoviePipeline::FObjectIdAccelerationData AccelData = ManifestAnnotation.GetAnnotation(ObjectIdArgs->RenderPassNode.Get());
+			UpdateCryptomatteMetadata(AccelData, SampleStatePayload->TraversalContext.RenderDataIdentifier.RendererName, SampleStatePayload->AdditionalFileMetadata);
+
+			for (int32 Index = 0; Index < ObjectIdArgs->NumOutputLayers; Index++)
+			{
+				// We unfortunately can't share ownership of the payload from the last sample due to the changed pass identifiers and layer name.
+				TSharedRef<FMovieGraphSampleState, ESPMode::ThreadSafe> NewPayload = SampleStatePayload->Copy();
+				FMovieGraphRenderDataIdentifier NewIdentifier = FMovieGraphRenderDataIdentifier(SampleStatePayload->TraversalContext.RenderDataIdentifier);
+				NewIdentifier.SubResourceName =  NewIdentifier.RendererName + FString::Printf(TEXT("%02d"), Index);
+				NewPayload->TraversalContext.RenderDataIdentifier = NewIdentifier;
+
+				// Update the layer name to be exactly what is needed for the Cryptomatte specification. The "name" in the metadata must match up
+				// exactly with the layer names in the output file (with a number at the end, eg "<LayerName>01").
+				NewPayload->LayerNameOverride = NewIdentifier.SubResourceName;
+
+				// Bump the sort order for the layer. An Object ID layer shouldn't be the first in the output format. 100 here is just an arbitrary
+				// number to force the Object ID layers to the end.
+				NewPayload->CompositingSortOrder += 100;
+
+				// Don't allow OCIO, otherwise Cryptomatte colors will not be correct.
+				NewPayload->bAllowOCIO = false;
+
+				TUniquePtr<TImagePixelData<FLinearColor>> FinalPixelData = MakeUnique<TImagePixelData<FLinearColor>>(FIntPoint(FullSizeX, FullSizeY), MoveTemp(OutputLayers[Index]), NewPayload);
+
+				// Send each layer to the Output Builder
+				ObjectIdArgs->OutputMerger->OnCompleteRenderPassDataAvailable_AnyThread(MoveTemp(FinalPixelData));
+			}
+			
+			// Free the memory in the accumulator now that we've extracted all
+			AccumulatorPin->Reset();
+		}
+	}
+}
+
+FMovieGraphObjectIdPass::FMovieGraphObjectIdPass()
+{
+}
+
+void FMovieGraphObjectIdPass::Setup(TWeakObjectPtr<UMovieGraphDefaultRenderer> InRenderer, TWeakObjectPtr<UMovieGraphImagePassBaseNode> InRenderPassNode, const FMovieGraphRenderPassLayerData& InLayer)
+{
+	FMovieGraphImagePassBase::Setup(InRenderer, InRenderPassNode, InLayer);
+
+	LayerData = InLayer;
+
+	// We output three layers. This equals 6 "ranks" in the Cryptomatte spec.
+	for (int32 Index = 0; Index < 3; Index++)
+	{
+		FMovieGraphRenderDataIdentifier NewIdentifier;
+
+		NewIdentifier.RootBranchName = LayerData.BranchName;
+		NewIdentifier.LayerName = LayerData.LayerName;
+		NewIdentifier.RendererName = InRenderPassNode->GetRendererName();
+		NewIdentifier.SubResourceName = NewIdentifier.RendererName + FString::Printf(TEXT("%02d"), Index);
+
+		UE::MovieGraph::DefaultRenderer::FCameraInfo CameraInfo = InRenderer->GetCameraInfo(LayerData.CameraIdentifier);
+		NewIdentifier.CameraName =  CameraInfo.CameraName;
+		
+		RenderDataIdentifiers.Add(NewIdentifier);
+	}
+
+	// The deferred render pass base class needs a valid RenderDataIdentifier, so the first Object ID render data identifier will be used. All of
+	// the Object ID render data identifiers are the same, other than the sub resource name, so this should be fine.
+	RenderDataIdentifier = RenderDataIdentifiers[0];
+
+	UE::MoviePipeline::FObjectIdAccelerationData AccelData = UE::MoviePipeline::FObjectIdAccelerationData();
+
+	// Static metadata needed for Cryptomatte
+	const uint32 NameHash = ::MoviePipeline::HashNameToId(TCHAR_TO_UTF8(*InRenderPassNode->GetRendererName()));
+	FString PassIdentifierHashAsShortString = FString::Printf(TEXT("%08x"), NameHash);
+	PassIdentifierHashAsShortString.LeftInline(7);
+
+	AccelData.PassIdentifierHashAsShortString = PassIdentifierHashAsShortString;
+	AccelData.JsonManifest = MakeShared<FJsonObject>();
+	AccelData.Cache = MakeShared<TMap<int32, UE::MoviePipeline::FMoviePipelineHitProxyCacheValue>>();
+	AccelData.Cache->Reserve(1000);
+
+	// Update the Cryptomatte manifest
+	{
+		// Add our default to the manifest
+		static const uint32 DefaultHash = ::MoviePipeline::HashNameToId(TCHAR_TO_UTF8(TEXT("default")));
+		AccelData.JsonManifest->SetStringField(TEXT("default"), FString::Printf(TEXT("%08x"), DefaultHash));
+
+		// Add the HitProxies
+		const UMovieGraphObjectIdNode* ObjectIdNode = Cast<UMovieGraphObjectIdNode>(LayerData.RenderPassNode);
+		check(ObjectIdNode);
+		UpdateManifestAccelerationData(AccelData, ObjectIdNode->IdType);
+	}
+	
+	GetManifestAnnotation().AddAnnotation(InLayer.RenderPassNode.Get(), AccelData);
+
+	SceneViewState.Allocate(InRenderer->GetWorld()->GetFeatureLevel());
+
+	// The InRenderPassNode is not initialized with user's config. Use InLayer to 
+	// initialize the frames to delay for post submission.
+	UMovieGraphRenderPassNode* RenderPassNode = InLayer.RenderPassNode.Get();
+	FramesToDelayPostSubmission = RenderPassNode ? RenderPassNode->GetCoolingDownFrameCount() : 0;
+}
+
+void FMovieGraphObjectIdPass::Teardown()
+{
+	FMovieGraphDeferredPass::Teardown();
+
+	GetManifestAnnotation().RemoveAnnotation(LayerData.RenderPassNode.Get());
+}
+
+void FMovieGraphObjectIdPass::GatherOutputPasses(UMovieGraphEvaluatedConfig* InConfig, TArray<FMovieGraphRenderDataIdentifier>& OutExpectedPasses) const
+{
+	OutExpectedPasses.Append(RenderDataIdentifiers);
+}
+
+UMovieGraphImagePassBaseNode* FMovieGraphObjectIdPass::GetParentNode(UMovieGraphEvaluatedConfig* InConfig) const
+{
+	constexpr bool bIncludeCDOs = true;
+	UMovieGraphObjectIdNode* ParentNode = InConfig->GetSettingForBranch<UMovieGraphObjectIdNode>(GetBranchName(), bIncludeCDOs);
+	if (!ensureMsgf(ParentNode, TEXT("ObjectIdPass should not exist without parent node in graph.")))
+	{
+		return nullptr;
+	}
+
+	return ParentNode;
+}
+
+TSharedRef<MoviePipeline::IMoviePipelineAccumulationArgs> FMovieGraphObjectIdPass::GetOrCreateAccumulator(TObjectPtr<UMovieGraphDefaultRenderer> InGraphRenderer, const UE::MovieGraph::FMovieGraphSampleState& InSampleState) const
+{
+	const FMoviePipelineAccumulatorPoolPtr SampleAccumulatorPool = InGraphRenderer->GetOrCreateAccumulatorPool<FMaskOverlappedAccumulator>();
+	const UE::MovieGraph::DefaultRenderer::FSurfaceAccumulatorPool::FInstancePtr AccumulatorInstance = SampleAccumulatorPool->GetAccumulatorInstance_GameThread<FMaskOverlappedAccumulator>(InSampleState.TraversalContext.Time.OutputFrameNumber, InSampleState.TraversalContext.RenderDataIdentifier);
+
+	TSharedRef<FMovieGraphObjectIdMaskSampleAccumulationArgs> AccumulationArgs = MakeShared<FMovieGraphObjectIdMaskSampleAccumulationArgs>();
+	AccumulationArgs->OutputMerger = InGraphRenderer->GetOwningGraph()->GetOutputMerger();
+	AccumulationArgs->ImageAccumulator = StaticCastSharedPtr<FMaskOverlappedAccumulator>(AccumulatorInstance->Accumulator);
+	AccumulationArgs->AccumulatorInstance = SampleAccumulatorPool->GetAccumulatorInstance_GameThread<FMaskOverlappedAccumulator>(InSampleState.TraversalContext.Time.OutputFrameNumber, InSampleState.TraversalContext.RenderDataIdentifier);
+	AccumulationArgs->NumOutputLayers = RenderDataIdentifiers.Num();
+	AccumulationArgs->CacheData = MakeShared<TMap<int32, UE::MoviePipeline::FMoviePipelineHitProxyCacheValue>>(*GetManifestAnnotation().GetAnnotation(LayerData.RenderPassNode.Get()).Cache);
+	AccumulationArgs->RenderPassNode = LayerData.RenderPassNode;
+
+	return AccumulationArgs;
+}
+
+UE::MovieGraph::Rendering::FMovieGraphImagePassBase::FAccumulatorSampleFunc FMovieGraphObjectIdPass::GetAccumulateSampleFunction() const
+{
+	return UE::MovieGraph::AccumulateSampleObjectId_TaskThread;
+}
+
+UE::MovieGraph::DefaultRenderer::FRenderTargetInitParams FMovieGraphObjectIdPass::GetRenderTargetInitParams(const FMovieGraphTimeStepData& InTimeData, const FIntPoint& InResolution)
+{
+	UE::MovieGraph::DefaultRenderer::FRenderTargetInitParams InitParams;
+
+	// Ensure there's no gamma in the render target otherwise the HitProxy color IDs don't round-trip properly.
+	InitParams.Size = InResolution;
+	InitParams.TargetGamma = 0.f;
+	InitParams.bForceLinearGamma = true;
+	InitParams.PixelFormat = PF_B8G8R8A8;
+
+	return InitParams;
+}
+
+FUObjectAnnotationSparse<UE::MoviePipeline::FObjectIdAccelerationData, true>& FMovieGraphObjectIdPass::GetManifestAnnotation()
+{
+	return UMovieGraphObjectIdNode::GetManifestAnnotation();
+}
\ No newline at end of file
diff --git a/Engine/Plugins/MovieScene/MoviePipelineMaskRenderPass/Source/MoviePipelineMaskRenderPass/Private/Graph/MovieGraphObjectIdPass.h b/Engine/Plugins/MovieScene/MoviePipelineMaskRenderPass/Source/MoviePipelineMaskRenderPass/Private/Graph/MovieGraphObjectIdPass.h
new file mode 100644
--- /dev/null
+++ b/Engine/Plugins/MovieScene/MoviePipelineMaskRenderPass/Source/MoviePipelineMaskRenderPass/Private/Graph/MovieGraphObjectIdPass.h
@@ -0,0 +1,51 @@
+﻿// Copyright Epic Games, Inc. All Rights Reserved.
+
+#pragma once
+
+#include "MoviePipelineObjectIdUtils.h"
+#include "Graph/Renderers/MovieGraphDeferredPass.h"
+#include "Graph/Renderers/MovieGraphImagePassBase.h"
+#include "UObject/UObjectAnnotation.h"
+
+struct FMovieGraphObjectIdMaskSampleAccumulationArgs : MoviePipeline::IMoviePipelineAccumulationArgs
+{
+	TWeakPtr<MoviePipeline::IMoviePipelineOverlappedAccumulator, ESPMode::ThreadSafe> ImageAccumulator;
+	UE::MovieGraph::DefaultRenderer::FSurfaceAccumulatorPool::FInstancePtr AccumulatorInstance;
+
+	/** The number of layers that the accumulator will be generating. */
+	int32 NumOutputLayers;
+
+	/** The mapping of a HitProxy index to the data associated with the HitProxy. */
+	TSharedPtr<TMap<int32, UE::MoviePipeline::FMoviePipelineHitProxyCacheValue>> CacheData;
+
+	/** The output merger that the accumulator should send pixel data to. */
+	TSharedPtr<UE::MovieGraph::IMovieGraphOutputMerger, ESPMode::ThreadSafe> OutputMerger;
+
+	/** The node that is using this accumulator. */
+	TWeakObjectPtr<UMovieGraphRenderPassNode> RenderPassNode;
+};
+
+/** The pass type that is capable of generating Object ID (Cryptomatte) data. */
+struct FMovieGraphObjectIdPass : UE::MovieGraph::Rendering::FMovieGraphDeferredPass
+{
+	FMovieGraphObjectIdPass();
+	
+	// FMovieGraphImagePassBase Interface
+	virtual void Setup(TWeakObjectPtr<UMovieGraphDefaultRenderer> InRenderer, TWeakObjectPtr<UMovieGraphImagePassBaseNode> InRenderPassNode, const FMovieGraphRenderPassLayerData& InLayer) override;
+	virtual void Teardown() override;
+	virtual void GatherOutputPasses(UMovieGraphEvaluatedConfig* InConfig, TArray<FMovieGraphRenderDataIdentifier>& OutExpectedPasses) const override;
+	virtual UMovieGraphImagePassBaseNode* GetParentNode(UMovieGraphEvaluatedConfig* InConfig) const override;
+	virtual TSharedRef<MoviePipeline::IMoviePipelineAccumulationArgs> GetOrCreateAccumulator(TObjectPtr<UMovieGraphDefaultRenderer> InGraphRenderer, const UE::MovieGraph::FMovieGraphSampleState& InSampleState) const override;
+	virtual FAccumulatorSampleFunc GetAccumulateSampleFunction() const override;
+	// ~FMovieGraphImagePassBase Interface
+
+protected:
+	virtual UE::MovieGraph::DefaultRenderer::FRenderTargetInitParams GetRenderTargetInitParams(const FMovieGraphTimeStepData& InTimeData, const FIntPoint& InResolution) override;
+
+private:
+	static FUObjectAnnotationSparse<UE::MoviePipeline::FObjectIdAccelerationData, true>& GetManifestAnnotation();
+
+protected:
+	/** The identifiers for all object ID layers that will be generated. */
+	TArray<FMovieGraphRenderDataIdentifier> RenderDataIdentifiers;
+};
\ No newline at end of file
diff --git a/Engine/Plugins/MovieScene/MoviePipelineMaskRenderPass/Source/MoviePipelineMaskRenderPass/Private/MoviePipelineObjectIdPass.cpp b/Engine/Plugins/MovieScene/MoviePipelineMaskRenderPass/Source/MoviePipelineMaskRenderPass/Private/MoviePipelineObjectIdPass.cpp
--- a/Engine/Plugins/MovieScene/MoviePipelineMaskRenderPass/Source/MoviePipelineMaskRenderPass/Private/MoviePipelineObjectIdPass.cpp
+++ b/Engine/Plugins/MovieScene/MoviePipelineMaskRenderPass/Source/MoviePipelineMaskRenderPass/Private/MoviePipelineObjectIdPass.cpp
@@ -1,6 +1,7 @@
 // Copyright Epic Games, Inc. All Rights Reserved.
 
 #include "MoviePipelineObjectIdPass.h"
+#include "MoviePipelineObjectIdUtils.h"
 #include "CanvasTypes.h"
 #include "MoviePipelineOutputBuilder.h"
 #include "Engine/TextureRenderTarget2D.h"
@@ -27,77 +28,6 @@
 #include UE_INLINE_GENERATED_CPP_BY_NAME(MoviePipelineObjectIdPass)
 
 DECLARE_CYCLE_STAT(TEXT("STAT_MoviePipeline_AccumulateMaskSample_TT"), STAT_AccumulateMaskSample_TaskThread, STATGROUP_MoviePipeline);
-namespace UE
-{
-namespace MoviePipeline
-{
-	struct FMoviePipelineHitProxyCacheKey
-	{
-		const AActor* Actor;
-		const UPrimitiveComponent* PrimComponent;
-
-		FMoviePipelineHitProxyCacheKey(const AActor* InActor, const UPrimitiveComponent* InComponent)
-			: Actor(InActor)
-			, PrimComponent(InComponent)
-		{}
-
-		friend inline uint32 GetTypeHash(const FMoviePipelineHitProxyCacheKey& Key)
-		{
-			return HashCombine(PointerHash(Key.Actor), PointerHash(Key.PrimComponent));
-		}
-
-		bool operator==(const FMoviePipelineHitProxyCacheKey& Other) const
-		{
-			return (Actor == Other.Actor) && (PrimComponent == Other.PrimComponent);
-		}
-	};
-
-	struct FMoviePipelineHitProxyCacheValue
-	{
-		const AActor* Actor;
-		const UPrimitiveComponent* PrimComponent;
-		int32 SectionIndex;
-		int32 MaterialIndex;
-		float Hash;
-		FString HashAsString;
-		FString ProxyName;
-	};
-
-	struct FObjectIdAccelerationData
-	{
-		FObjectIdAccelerationData()
-		{}
-
-		FORCEINLINE bool IsDefault()
-		{
-			return !Cache.IsValid()
-				&& !JsonManifest.IsValid()
-				&& JsonManifestCachedOutput.Len() == 0
-				&& PassIdentifierHashAsShortString.Len() == 0;
-		}
-
-		TSharedPtr<TMap<int32, FMoviePipelineHitProxyCacheValue>> Cache;
-
-		// Json Manifest Object
-		TSharedPtr<FJsonObject> JsonManifest;
-
-		// Cached version of the serialized Json Manifest
-		FString JsonManifestCachedOutput;
-
-		FString PassIdentifierHashAsShortString;
-	};
-
-	struct FObjectIdMaskSampleAccumulationArgs
-	{
-	public:
-		TSharedPtr<FMaskOverlappedAccumulator, ESPMode::ThreadSafe> Accumulator;
-		TSharedPtr<FMoviePipelineOutputMerger, ESPMode::ThreadSafe> OutputMerger;
-		int32 NumOutputLayers;
-		TSharedPtr<TMap<int32, FMoviePipelineHitProxyCacheValue>> CacheData;
-	};
-}
-}
-
 
 static FUObjectAnnotationSparse<UE::MoviePipeline::FObjectIdAccelerationData, true> ManifestAnnotation;
 
@@ -216,100 +146,6 @@ void UMoviePipelineObjectIdRenderPass::GetViewShowFlags(FEngineShowFlags& OutSho
 	OutViewModeIndex = EViewModeIndex::VMI_Unlit;
 }
 
-FString UMoviePipelineObjectIdRenderPass::ResolveProxyIdGroup(const AActor* InActor, const UPrimitiveComponent* InPrimComponent, const int32 InMaterialIndex, const int32 InSectionIndex) const
-{
-	// If it doesn't exist in the cache already, then we will do the somewhat expensive of building the string and hashing it.
-	TStringBuilder<128> StringBuilder;
-
-	FName FolderPath = InActor->GetFolderPath();
-
-	// If they don't want the hierarchy, we'll just set this to empty string.
-	if (IdType == EMoviePipelineObjectIdPassIdType::Actor)
-	{
-		FolderPath = NAME_None;
-	}
-
-	switch (IdType)
-	{
-	case EMoviePipelineObjectIdPassIdType::Layer:
-	{
-		if (InActor->Layers.Num() > 0)
-		{
-			StringBuilder.Append(*InActor->Layers[0].ToString());
-		}
-		break;
-	}
-	case EMoviePipelineObjectIdPassIdType::Folder:
-	{
-		if (!FolderPath.IsNone())
-		{
-			StringBuilder.Append(*FolderPath.ToString());
-		}
-		break;
-	}
-
-	case EMoviePipelineObjectIdPassIdType::Material:
-	{
-		if (InPrimComponent->GetNumMaterials() > 0)
-		{
-			UMaterialInterface* MaterialInterface = InPrimComponent->GetMaterial(FMath::Clamp(InMaterialIndex, 0, InPrimComponent->GetNumMaterials() - 1));
-
-			// This collapses dynamic material instances back into their parent asset so we don't end up with 'MaterialInstanceDynamic_1' instead of MI_Foo
-			if (UMaterialInstanceDynamic* AsDynamicMaterialInstance = Cast<UMaterialInstanceDynamic>(MaterialInterface))
-			{
-				if (AsDynamicMaterialInstance->Parent)
-				{
-					StringBuilder.Append(*AsDynamicMaterialInstance->Parent->GetName());
-				}
-				else
-				{
-					StringBuilder.Append(*AsDynamicMaterialInstance->GetName());
-				}
-			}
-			else if (UMaterialInstance* AsMaterialInstance = Cast<UMaterialInstance>(MaterialInterface))
-			{
-				StringBuilder.Append(*MaterialInterface->GetName());
-			}
-			else if (MaterialInterface && MaterialInterface->GetMaterial())
-			{
-				StringBuilder.Append(*MaterialInterface->GetMaterial()->GetName());
-			}
-		}
-		break;
-	}
-	case EMoviePipelineObjectIdPassIdType::Actor:
-	case EMoviePipelineObjectIdPassIdType::ActorWithHierarchy:
-	{
-		// Folder Path will be NAME_None for root objects and for the "Actor" group type.
-		if (!FolderPath.IsNone())
-		{
-			StringBuilder.Append(*FolderPath.ToString());
-			StringBuilder.Append(TEXT("/"));
-		}
-		StringBuilder.Append(*InActor->GetActorLabel());
-		break;
-	}
-	case EMoviePipelineObjectIdPassIdType::Full:
-	{
-		// Full gives as much detail as we can - per folder, per actor, per component, per material
-		if (!FolderPath.IsNone())
-		{
-			StringBuilder.Append(*FolderPath.ToString());
-			StringBuilder.Append(TEXT("/"));
-		}
-		StringBuilder.Appendf(TEXT("%s.%s[%d.%d]"), *InActor->GetActorLabel(), *GetNameSafe(InPrimComponent), InMaterialIndex, InSectionIndex);
-		break;
-	}
-	}
-
-	if (StringBuilder.Len() == 0)
-	{
-		StringBuilder.Append(TEXT("default"));
-	}
-
-	return StringBuilder.ToString();
-}
-
 void UMoviePipelineObjectIdRenderPass::RenderSample_GameThreadImpl(const FMoviePipelineRenderPassMetrics& InSampleState)
 {
 	// Object Ids have no history buffer so no need to render when we're going to discard.
@@ -336,131 +172,16 @@ void UMoviePipelineObjectIdRenderPass::RenderSample_GameThreadImpl(const FMovieP
 	FCanvas Canvas = FCanvas(RenderTarget, nullptr, GetPipeline()->GetWorld(), ViewFamily->GetFeatureLevel(), FCanvas::CDM_DeferDrawing, 1.0f);
 	GetRendererModule().BeginRenderingViewFamily(&Canvas, ViewFamily.Get());
 
-	// The Hitproxy array gets invalidated quite often, so the results are no longer valid in the accumulation thread.
-	// To solve this, we will cache the required info on the game thread and pass the required info along with the render so that
-	// it stays in sync with what was actually rendered. Additionally, we cache the hashes between frames as they will be largely 
-	// the same between each frame.
-	const TSparseArray<HHitProxy*>& AllHitProxies = GetAllHitProxies();
-
-	std::atomic<int32> NumCacheHits(0);
-	std::atomic<int32> NumCacheMisses(0);
-	std::atomic<int32> NumCacheUpdates(0);
-
 	// Update the data in place, no need to copy back to the annotation.
 	UE::MoviePipeline::FObjectIdAccelerationData AccelData = ManifestAnnotation.GetAnnotation(this);
-	const double CacheStartTime = FPlatformTime::Seconds();
-
-	for (typename TSparseArray<HHitProxy*>::TConstIterator It(AllHitProxies); It; ++It)
-	{
-		HActor* ActorHitProxy = HitProxyCast<HActor>(*It);
-		HInstancedStaticMeshInstance* FoliageHitProxy = HitProxyCast<HInstancedStaticMeshInstance>(*It);
-
-		const AActor* ProxyActor = nullptr;
-		const UPrimitiveComponent* ProxyComponent = nullptr;
-		int32 ProxySectionIndex = -1;
-		int32 ProxyMaterialIndex = -1;
-
-		if (ActorHitProxy && IsValid(ActorHitProxy->Actor) && IsValid(ActorHitProxy->PrimComponent))
-		{
-			ProxyActor = ActorHitProxy->Actor;
-			ProxyComponent = ActorHitProxy->PrimComponent;
-			ProxySectionIndex = ActorHitProxy->SectionIndex;
-			ProxyMaterialIndex = ActorHitProxy->MaterialIndex;
-		}
-		else if (FoliageHitProxy && IsValid(FoliageHitProxy->Component))
-		{
-			ProxyActor = FoliageHitProxy->Component->GetOwner();
-			ProxyComponent = FoliageHitProxy->Component;
-			ProxySectionIndex = FoliageHitProxy->InstanceIndex;
-		}
-
-		if(ProxyActor && ProxyComponent)
-		{
-			// We assume names to be stable within a shot. This is technically incorrect if you were to 
-			// rename an actor mid-frame but using this assumption allows us to skip calculating the string
-			// name every frame.
-			UE::MoviePipeline::FMoviePipelineHitProxyCacheValue* CacheEntry = nullptr;
-			FColor Color = (*It)->Id.GetColor();
-			int32 IdToInt = ((int32)Color.R << 16) | ((int32)Color.G << 8) | ((int32)Color.B << 0);
-			{
-				CacheEntry = AccelData.Cache->Find(IdToInt);
-			}
-
-
-			if (CacheEntry)
-			{
-				// The cache could be out of date since it's only an index. We'll double check that the actor and component
-				// are the same and assume if they are, the cache is still valid.
-				const bool bSameActor = CacheEntry->Actor == ProxyActor;
-				const bool bSameComp = CacheEntry->PrimComponent == ProxyComponent;
-				const bool bSameSection = CacheEntry->SectionIndex == ProxySectionIndex;
-				const bool bSameMaterial = CacheEntry->MaterialIndex == ProxyMaterialIndex;
-
-				if (bSameActor && bSameComp && bSameSection && bSameMaterial)
-				{
-					NumCacheHits++;
-					continue;
-				}
-				NumCacheUpdates++;
-			}
-			NumCacheMisses++;
-
-			FString ProxyIdName = ResolveProxyIdGroup(ProxyActor, ProxyComponent, ProxyMaterialIndex, ProxySectionIndex);
-			
-			// We hash the string and printf it here to reduce allocations later, even though it makes this loop ~% more expensive.
-			uint32 Hash = MoviePipeline::HashNameToId(TCHAR_TO_UTF8(*ProxyIdName));
-			FString HashAsString = FString::Printf(TEXT("%08x"), Hash);
-
-			{
-				UE::MoviePipeline::FMoviePipelineHitProxyCacheValue& NewCacheEntry = AccelData.Cache->Add(IdToInt);
-				NewCacheEntry.ProxyName = ProxyIdName;
-				NewCacheEntry.Hash = *(float*)(&Hash);
-				NewCacheEntry.Actor = ProxyActor;
-				NewCacheEntry.PrimComponent = ProxyComponent;
-				NewCacheEntry.SectionIndex = ProxySectionIndex;
-				NewCacheEntry.MaterialIndex = ProxyMaterialIndex;
-
-				// Add the object to the manifest. Done here because this takes ~170ms a frame for 700 objects.
-				// May as well only take that hit once per shot. This will add or update an existing field.
-				AccelData.JsonManifest->SetStringField(ProxyIdName, HashAsString);
-
-				// Only move it after we've used it to update the Json Manifest.
-				NewCacheEntry.HashAsString = MoveTemp(HashAsString);
-			}
-		}
-	}
-
-	const double CacheEndTime = FPlatformTime::Seconds();
-	const float ElapsedMs = float((CacheEndTime - CacheStartTime) * 1000.0f);
-
-	//{
-		const double JsonBeginTime = FPlatformTime::Seconds();
-
-		// We only update the serialized manifest file if something has changed since it's slow.
-		if (NumCacheMisses.load() > 0)
-		{
-			AccelData.JsonManifestCachedOutput.Empty();
-			TSharedRef<TJsonWriter<>> Writer = TJsonWriterFactory<>::Create(&AccelData.JsonManifestCachedOutput);
-			FJsonSerializer::Serialize(AccelData.JsonManifest.ToSharedRef(), Writer);
-		}
-
-		FString PassIdentifierHashAsShortString = AccelData.PassIdentifierHashAsShortString;
-		InOutSampleState.OutputState.FileMetadata.Add(FString::Printf(TEXT("cryptomatte/%s/manifest"), *PassIdentifierHashAsShortString), AccelData.JsonManifestCachedOutput);
-		InOutSampleState.OutputState.FileMetadata.Add(FString::Printf(TEXT("cryptomatte/%s/name"), *PassIdentifierHashAsShortString), PassIdentifier.Name);
-		InOutSampleState.OutputState.FileMetadata.Add(FString::Printf(TEXT("cryptomatte/%s/hash"), *PassIdentifierHashAsShortString), TEXT("MurmurHash3_32"));
-		InOutSampleState.OutputState.FileMetadata.Add(FString::Printf(TEXT("cryptomatte/%s/conversion"), *PassIdentifierHashAsShortString), TEXT("uint32_to_float32"));
-		const double JsonEndTime = FPlatformTime::Seconds();
-		const float ElapsedJsonMs = float((JsonEndTime - JsonBeginTime) * 1000.0f);
-	//}
-
-	UE_LOG(LogMovieRenderPipeline, VeryVerbose, TEXT("Cache Size: %d NumCacheHits: %d NumCacheMisses: %d NumCacheUpdates: %d CacheDuration: %8.2fms JsonDuration: %8.2fms"), AccelData.Cache->Num(), NumCacheHits.load(), NumCacheMisses.load(), NumCacheUpdates.load(), ElapsedMs, ElapsedJsonMs);
-	
-	// Update the annotation with new cached data
+	UE::MoviePipeline::UpdateManifestAccelerationData(AccelData, IdType);
 	ManifestAnnotation.AddAnnotation(this, AccelData);
 
+	// Update ObjectID-related file metadata
+	UpdateCryptomatteMetadata(AccelData, PassIdentifier.Name, InOutSampleState.OutputState.FileMetadata);
+
 	// Main Render Pass
 	{
-
 		// Readback + Accumulate.
 		TSharedRef<FImagePixelDataPayload, ESPMode::ThreadSafe> FramePayload = MakeShared<FImagePixelDataPayload, ESPMode::ThreadSafe>();
 		FramePayload->PassIdentifier = PassIdentifier;
@@ -485,7 +206,6 @@ void UMoviePipelineObjectIdRenderPass::RenderSample_GameThreadImpl(const FMovieP
 		auto Callback = [this, FramePayload, AccumAgs = MoveTemp(AccumulationArgs), SampleAccumulator](TUniquePtr<FImagePixelData>&& InPixelData) mutable
 		{
 			bool bFinalSample = FramePayload->IsLastTile() && FramePayload->IsLastTemporalSample();
-			bool bFirstSample = FramePayload->IsFirstTile() && FramePayload->IsFirstTemporalSample();
 
 			FMoviePipelineBackgroundAccumulateTask Task;
 			Task.LastCompletionEvent = SampleAccumulator->TaskPrereq;
@@ -552,9 +272,6 @@ namespace MoviePipeline
 		FImagePixelDataPayload* FramePayload = InPixelData->GetPayload<FImagePixelDataPayload>();
 		check(FramePayload);
 
-		static const uint32 DefaultHash = HashNameToId(TCHAR_TO_UTF8(TEXT("default")));
-		static const float DefaultHashAsFloat = *(float*)(&DefaultHash);
-
 		// Writing tiles can be useful for debug reasons. These get passed onto the output every frame.
 		if (FramePayload->SampleState.bWriteSampleToDisk)
 		{
@@ -565,7 +282,6 @@ namespace MoviePipeline
 			InParams.OutputMerger->OnSingleSampleDataAvailable_AnyThread(MoveTemp(SampleData));
 		}
 
-
 		// For the first sample in a new output, we allocate memory
 		if (FramePayload->IsFirstTile() && FramePayload->IsFirstTemporalSample())
 		{
@@ -591,31 +307,7 @@ namespace MoviePipeline
 			IdData.SetNumUninitialized(RawSize.X * RawSize.Y);
 
 			// Remap hitproxy id into precalculated Cryptomatte hash
-			ParallelFor(RawSize.Y,
-				[&](int32 ScanlineIndex = 0)
-				{
-					for (int64 Index = 0; Index < RawSize.X; Index++)
-					{
-						int64 DstIndex = int64(ScanlineIndex) * int64(RawSize.X) + int64(Index);
-						// Turn the FColor into an integer index
-						const FColor* Color = &RawDataPtr[DstIndex];
-
-						int32 HitProxyIndex = ((int32)Color->R << 16) | ((int32)Color->G << 8) | ((int32)Color->B << 0);
-
-						float Hash = DefaultHashAsFloat;
-						const UE::MoviePipeline::FMoviePipelineHitProxyCacheValue* CachedValue = InParams.CacheData->Find(HitProxyIndex);
-						if (CachedValue)
-						{
-							Hash = CachedValue->Hash;
-						}
-						else
-						{
-							UE_LOG(LogMovieRenderPipeline, VeryVerbose, TEXT("Failed to find cache data for Hitproxy! Id: %d"), HitProxyIndex);
-						}
-
-						IdData[DstIndex] = Hash;
-					}
-				});
+			RemapHitProxyIdToCryptomatteHash(RawSize, RawDataPtr, InParams.CacheData, IdData);
 			const double RemapEndTime = FPlatformTime::Seconds();
 			const float ElapsedRemapMs = float((RemapEndTime - RemapBeginTime) * 1000.0f);
 
diff --git a/Engine/Plugins/MovieScene/MoviePipelineMaskRenderPass/Source/MoviePipelineMaskRenderPass/Private/MoviePipelineObjectIdPass.h b/Engine/Plugins/MovieScene/MoviePipelineMaskRenderPass/Source/MoviePipelineMaskRenderPass/Private/MoviePipelineObjectIdPass.h
--- a/Engine/Plugins/MovieScene/MoviePipelineMaskRenderPass/Source/MoviePipelineMaskRenderPass/Private/MoviePipelineObjectIdPass.h
+++ b/Engine/Plugins/MovieScene/MoviePipelineMaskRenderPass/Source/MoviePipelineMaskRenderPass/Private/MoviePipelineObjectIdPass.h
@@ -4,22 +4,7 @@
 #include "MoviePipelineImagePassBase.h"
 #include "MoviePipelineObjectIdPass.generated.h"
 
-UENUM(BlueprintType)
-enum class EMoviePipelineObjectIdPassIdType : uint8
-{
-	/** As much information as the renderer can provide - unique per material per primitive in the world. */
-	Full,
-	/** Grouped by material name. This means different objects that use the same material will be merged. */
-	Material,
-	/** Grouped by Actor Name, all materials for a given actor are merged together, and all actors with that name are merged together as well. */
-	Actor,
-	/** Grouped by Actor Name and Folder Hierarchy. This means actors with the same name in different folders will not be merged together. */
-	ActorWithHierarchy,
-	/** Grouped by Folder Name. All actors within a given folder hierarchy in the World Outliner are merged together. */
-	Folder,
-	/** Primary Layer. This is the first layer found in the AActor::Layers array. May not do what you expect if an actor belongs to multiple layers. */
-	Layer
-};
+enum class EMoviePipelineObjectIdPassIdType : uint8;
 
 UCLASS(BlueprintType)
 class UMoviePipelineObjectIdRenderPass : public UMoviePipelineImagePassBase
@@ -43,7 +28,6 @@ public:
 
 protected:
 	void PostRendererSubmission(const FMoviePipelineRenderPassMetrics& InSampleState);
-	FString ResolveProxyIdGroup(const AActor* InActor, const UPrimitiveComponent* InPrimComponent, const int32 InMaterialIndex, const int32 InSectionIndex) const;
 
 public:
 	UPROPERTY(EditAnywhere, BlueprintReadWrite, Category = "Settings")
diff --git a/Engine/Plugins/MovieScene/MoviePipelineMaskRenderPass/Source/MoviePipelineMaskRenderPass/Private/MoviePipelineObjectIdUtils.h b/Engine/Plugins/MovieScene/MoviePipelineMaskRenderPass/Source/MoviePipelineMaskRenderPass/Private/MoviePipelineObjectIdUtils.h
new file mode 100644
--- /dev/null
+++ b/Engine/Plugins/MovieScene/MoviePipelineMaskRenderPass/Source/MoviePipelineMaskRenderPass/Private/MoviePipelineObjectIdUtils.h
@@ -0,0 +1,360 @@
+﻿// Copyright Epic Games, Inc. All Rights Reserved.
+
+#pragma once
+
+#include "EngineUtils.h"
+#include "MovieRenderOverlappedMask.h"
+#include "Components/InstancedStaticMeshComponent.h"
+#include "Dom/JsonObject.h"
+#include "HitProxies.h"
+#include "Materials/Material.h"
+#include "Materials/MaterialInstance.h"
+#include "Materials/MaterialInstanceDynamic.h"
+#include "MoviePipelineHashUtils.h"
+#include "MovieRenderPipelineCoreModule.h"
+#include "Serialization/JsonSerializer.h"
+#include "Serialization/JsonWriter.h"
+
+#include "MoviePipelineObjectIdUtils.generated.h"
+
+extern const TSparseArray<HHitProxy*>& GetAllHitProxies();
+
+UENUM(BlueprintType)
+enum class EMoviePipelineObjectIdPassIdType : uint8
+{
+	/** As much information as the renderer can provide - unique per material per primitive in the world. */
+	Full,
+	
+	/** Grouped by material name. This means different objects that use the same material will be merged. */
+	Material,
+	
+	/** Grouped by Actor Name, all materials for a given actor are merged together, and all actors with that name are merged together as well. */
+	Actor,
+	
+	/** Grouped by Actor Name and Folder Hierarchy. This means actors with the same name in different folders will not be merged together. */
+	ActorWithHierarchy,
+	
+	/** Grouped by Folder Name. All actors within a given folder hierarchy in the World Outliner are merged together. */
+	Folder,
+	
+	/**
+	 * Grouped by Actor Layer (the first layer found in the AActor::Layers array). May not do what you expect if an actor belongs to multiple layers.
+	 * If used within a graph, this does NOT refer to the layer within the graph.
+	 *
+	 * In scripting, this option is referred to as "Layer" and not "ActorLayer".
+	 */
+	Layer	UMETA(DisplayName = "Actor Layer")
+};
+
+namespace UE::MoviePipeline
+{
+	struct FMoviePipelineHitProxyCacheKey
+	{
+		const AActor* Actor;
+		const UPrimitiveComponent* PrimComponent;
+
+		FMoviePipelineHitProxyCacheKey(const AActor* InActor, const UPrimitiveComponent* InComponent)
+			: Actor(InActor)
+			, PrimComponent(InComponent)
+		{}
+
+		friend inline uint32 GetTypeHash(const FMoviePipelineHitProxyCacheKey& Key)
+		{
+			return HashCombine(PointerHash(Key.Actor), PointerHash(Key.PrimComponent));
+		}
+
+		bool operator==(const FMoviePipelineHitProxyCacheKey& Other) const
+		{
+			return (Actor == Other.Actor) && (PrimComponent == Other.PrimComponent);
+		}
+	};
+
+	struct FMoviePipelineHitProxyCacheValue
+	{
+		const AActor* Actor;
+		const UPrimitiveComponent* PrimComponent;
+		int32 SectionIndex;
+		int32 MaterialIndex;
+		float Hash;
+		FString HashAsString;
+		FString ProxyName;
+	};
+
+	struct FObjectIdAccelerationData
+	{
+		FObjectIdAccelerationData()
+		{}
+
+		FORCEINLINE bool IsDefault() const
+		{
+			return !Cache.IsValid()
+				&& !JsonManifest.IsValid()
+				&& JsonManifestCachedOutput.Len() == 0
+				&& PassIdentifierHashAsShortString.Len() == 0;
+		}
+
+		/** Maps a HitProxy index to the data associated with the HitProxy. */
+		TSharedPtr<TMap<int32, FMoviePipelineHitProxyCacheValue>> Cache;
+
+		/** The Cryptomatte JSON manifest data. */
+		TSharedPtr<FJsonObject> JsonManifest;
+
+		/** A cached and serialized version of the JsonManifest data. */
+		FString JsonManifestCachedOutput;
+
+		/** The pass identifier that is used within Cryptomatte metadata keys like `cryptomatte/<pass_id>/name`. */
+		FString PassIdentifierHashAsShortString;
+	};
+
+	struct FObjectIdMaskSampleAccumulationArgs : ::MoviePipeline::IMoviePipelineAccumulationArgs
+	{
+		TSharedPtr<FMaskOverlappedAccumulator, ESPMode::ThreadSafe> Accumulator;
+		TSharedPtr<FMoviePipelineOutputMerger, ESPMode::ThreadSafe> OutputMerger;
+		int32 NumOutputLayers;
+		TSharedPtr<TMap<int32, FMoviePipelineHitProxyCacheValue>> CacheData;
+	};
+
+	static FString GenerateProxyIdGroup(const AActor* InActor, const UPrimitiveComponent* InPrimComponent, const EMoviePipelineObjectIdPassIdType IdType, const int32 InMaterialIndex, const int32 InSectionIndex)
+	{
+		// If it doesn't exist in the cache already, then we will do the somewhat expensive of building the string and hashing it.
+		TStringBuilder<128> StringBuilder;
+
+		FName FolderPath = InActor->GetFolderPath();
+
+		// If they don't want the hierarchy, we'll just set this to empty string.
+		if (IdType == EMoviePipelineObjectIdPassIdType::Actor)
+		{
+			FolderPath = NAME_None;
+		}
+
+		switch (IdType)
+		{
+		case EMoviePipelineObjectIdPassIdType::Layer:
+		{
+			if (InActor->Layers.Num() > 0)
+			{
+				StringBuilder.Append(*InActor->Layers[0].ToString());
+			}
+			break;
+		}
+		case EMoviePipelineObjectIdPassIdType::Folder:
+		{
+			if (!FolderPath.IsNone())
+			{
+				StringBuilder.Append(*FolderPath.ToString());
+			}
+			break;
+		}
+
+		case EMoviePipelineObjectIdPassIdType::Material:
+		{
+			if (InPrimComponent->GetNumMaterials() > 0)
+			{
+				UMaterialInterface* MaterialInterface = InPrimComponent->GetMaterial(FMath::Clamp(InMaterialIndex, 0, InPrimComponent->GetNumMaterials() - 1));
+
+				// This collapses dynamic material instances back into their parent asset so we don't end up with 'MaterialInstanceDynamic_1' instead of MI_Foo
+				if (const UMaterialInstanceDynamic* AsDynamicMaterialInstance = Cast<UMaterialInstanceDynamic>(MaterialInterface))
+				{
+					if (AsDynamicMaterialInstance->Parent)
+					{
+						StringBuilder.Append(*AsDynamicMaterialInstance->Parent->GetName());
+					}
+					else
+					{
+						StringBuilder.Append(*AsDynamicMaterialInstance->GetName());
+					}
+				}
+				else if (UMaterialInstance* AsMaterialInstance = Cast<UMaterialInstance>(MaterialInterface))
+				{
+					StringBuilder.Append(*MaterialInterface->GetName());
+				}
+				else if (MaterialInterface && MaterialInterface->GetMaterial())
+				{
+					StringBuilder.Append(*MaterialInterface->GetMaterial()->GetName());
+				}
+			}
+			break;
+		}
+		case EMoviePipelineObjectIdPassIdType::Actor:
+		case EMoviePipelineObjectIdPassIdType::ActorWithHierarchy:
+		{
+			// Folder Path will be NAME_None for root objects and for the "Actor" group type.
+			if (!FolderPath.IsNone())
+			{
+				StringBuilder.Append(*FolderPath.ToString());
+				StringBuilder.Append(TEXT("/"));
+			}
+			StringBuilder.Append(*InActor->GetActorLabel());
+			break;
+		}
+		case EMoviePipelineObjectIdPassIdType::Full:
+		{
+			// Full gives as much detail as we can - per folder, per actor, per component, per material
+			if (!FolderPath.IsNone())
+			{
+				StringBuilder.Append(*FolderPath.ToString());
+				StringBuilder.Append(TEXT("/"));
+			}
+			StringBuilder.Appendf(TEXT("%s.%s[%d.%d]"), *InActor->GetActorLabel(), *GetNameSafe(InPrimComponent), InMaterialIndex, InSectionIndex);
+			break;
+		}
+		}
+
+		if (StringBuilder.Len() == 0)
+		{
+			StringBuilder.Append(TEXT("default"));
+		}
+
+		return StringBuilder.ToString();
+	}
+
+	static void UpdateManifestAccelerationData(FObjectIdAccelerationData& InAccelData, const EMoviePipelineObjectIdPassIdType IdType)
+	{
+		// The HitProxy array gets invalidated quite often, so the results are no longer valid in the accumulation thread.
+		// To solve this, we will cache the required info on the game thread and pass the required info along with the render so that
+		// it stays in sync with what was actually rendered. Additionally, we cache the hashes between frames as they will be largely 
+		// the same between each frame.
+		const TSparseArray<HHitProxy*>& AllHitProxies = GetAllHitProxies();
+
+		std::atomic<int32> NumCacheHits(0);
+		std::atomic<int32> NumCacheMisses(0);
+		std::atomic<int32> NumCacheUpdates(0);
+
+		// Update the data in place, no need to copy back to the annotation.
+		const double CacheStartTime = FPlatformTime::Seconds();
+
+		for (typename TSparseArray<HHitProxy*>::TConstIterator It(AllHitProxies); It; ++It)
+		{
+			const HActor* ActorHitProxy = HitProxyCast<HActor>(*It);
+			const HInstancedStaticMeshInstance* FoliageHitProxy = HitProxyCast<HInstancedStaticMeshInstance>(*It);
+
+			const AActor* ProxyActor = nullptr;
+			const UPrimitiveComponent* ProxyComponent = nullptr;
+			int32 ProxySectionIndex = -1;
+			int32 ProxyMaterialIndex = -1;
+
+			if (ActorHitProxy && IsValid(ActorHitProxy->Actor) && IsValid(ActorHitProxy->PrimComponent))
+			{
+				ProxyActor = ActorHitProxy->Actor;
+				ProxyComponent = ActorHitProxy->PrimComponent;
+				ProxySectionIndex = ActorHitProxy->SectionIndex;
+				ProxyMaterialIndex = ActorHitProxy->MaterialIndex;
+			}
+			else if (FoliageHitProxy && IsValid(FoliageHitProxy->Component))
+			{
+				ProxyActor = FoliageHitProxy->Component->GetOwner();
+				ProxyComponent = FoliageHitProxy->Component;
+				ProxySectionIndex = FoliageHitProxy->InstanceIndex;
+			}
+
+			if (ProxyActor && ProxyComponent)
+			{
+				// We assume names to be stable within a shot. This is technically incorrect if you were to 
+				// rename an actor mid-frame but using this assumption allows us to skip calculating the string
+				// name every frame.
+				const FColor Color = (*It)->Id.GetColor();
+				int32 IdToInt = (static_cast<int32>(Color.R) << 16) | (static_cast<int32>(Color.G) << 8) | (static_cast<int32>(Color.B) << 0);
+
+				if (const FMoviePipelineHitProxyCacheValue* CacheEntry = InAccelData.Cache->Find(IdToInt))
+				{
+					// The cache could be out of date since it's only an index. We'll double check that the actor and component
+					// are the same and assume if they are, the cache is still valid.
+					const bool bSameActor = CacheEntry->Actor == ProxyActor;
+					const bool bSameComp = CacheEntry->PrimComponent == ProxyComponent;
+					const bool bSameSection = CacheEntry->SectionIndex == ProxySectionIndex;
+					const bool bSameMaterial = CacheEntry->MaterialIndex == ProxyMaterialIndex;
+
+					if (bSameActor && bSameComp && bSameSection && bSameMaterial)
+					{
+						++NumCacheHits;
+						continue;
+					}
+					++NumCacheUpdates;
+				}
+				++NumCacheMisses;
+
+				// We hash the string and printf it here to reduce allocations later, even though it makes this loop ~% more expensive.
+				{
+					FString ProxyIdName = GenerateProxyIdGroup(ProxyActor, ProxyComponent, IdType, ProxyMaterialIndex, ProxySectionIndex);
+					uint32 Hash = ::MoviePipeline::HashNameToId(TCHAR_TO_UTF8(*ProxyIdName));
+					FString HashAsString = FString::Printf(TEXT("%08x"), Hash);
+					
+					FMoviePipelineHitProxyCacheValue& NewCacheEntry = InAccelData.Cache->Add(IdToInt);
+					NewCacheEntry.ProxyName = ProxyIdName;
+					NewCacheEntry.Hash = *reinterpret_cast<float*>(&Hash);
+					NewCacheEntry.Actor = ProxyActor;
+					NewCacheEntry.PrimComponent = ProxyComponent;
+					NewCacheEntry.SectionIndex = ProxySectionIndex;
+					NewCacheEntry.MaterialIndex = ProxyMaterialIndex;
+
+					// Add the object to the manifest. Done here because this takes ~170ms a frame for 700 objects.
+					// May as well only take that hit once per shot. This will add or update an existing field.
+					InAccelData.JsonManifest->SetStringField(ProxyIdName, HashAsString);
+
+					// Only move it after we've used it to update the Json Manifest.
+					NewCacheEntry.HashAsString = MoveTemp(HashAsString);
+				}
+			}
+		}
+
+		const double CacheEndTime = FPlatformTime::Seconds();
+		const float ElapsedMs = static_cast<float>((CacheEndTime - CacheStartTime) * 1000.0f);
+
+		const double JsonBeginTime = FPlatformTime::Seconds();
+
+		// Update the cached JSON if needed.
+		// We only update the serialized manifest JSON if something has changed because serializing is slow.
+		if (NumCacheMisses.load() > 0)
+		{
+			InAccelData.JsonManifestCachedOutput.Empty();
+			
+			const TSharedRef<TJsonWriter<>> Writer = TJsonWriterFactory<>::Create(&InAccelData.JsonManifestCachedOutput);
+			FJsonSerializer::Serialize(InAccelData.JsonManifest.ToSharedRef(), Writer);
+		}
+
+		const double JsonEndTime = FPlatformTime::Seconds();
+		const float ElapsedJsonMs = static_cast<float>((JsonEndTime - JsonBeginTime) * 1000.0f);
+
+		UE_LOG(LogMovieRenderPipeline, VeryVerbose, TEXT("Cache Size: %d NumCacheHits: %d NumCacheMisses: %d NumCacheUpdates: %d CacheDuration: %8.2fms JsonDuration: %8.2fms"), InAccelData.Cache->Num(), NumCacheHits.load(), NumCacheMisses.load(), NumCacheUpdates.load(), ElapsedMs, ElapsedJsonMs);
+	}
+
+	static void RemapHitProxyIdToCryptomatteHash(const FIntPoint& InImageSize, const FColor* InHitProxyBuffer, const TSharedPtr<TMap<int32, FMoviePipelineHitProxyCacheValue>> InHitProxyCache, TArray64<float>& OutCryptomatteBuffer)
+	{
+		static const uint32 DefaultHash = ::MoviePipeline::HashNameToId(TCHAR_TO_UTF8(TEXT("default")));
+		static const float DefaultHashAsFloat = *(float*)(&DefaultHash);
+		
+		ParallelFor(InImageSize.Y,
+		[&](const int32 ScanlineIndex = 0)
+		{
+			for (int64 Index = 0; Index < InImageSize.X; Index++)
+			{
+				int64 DstIndex = int64(ScanlineIndex) * int64(InImageSize.X) + int64(Index);
+				const FColor* Color = &InHitProxyBuffer[DstIndex];
+
+				// Turn the FColor into an integer index
+				int32 HitProxyIndex = ((int32)Color->R << 16) | ((int32)Color->G << 8) | ((int32)Color->B << 0);
+
+				float Hash = DefaultHashAsFloat;
+				if (const FMoviePipelineHitProxyCacheValue* CachedValue = InHitProxyCache->Find(HitProxyIndex))
+				{
+					Hash = CachedValue->Hash;
+				}
+				else
+				{
+					UE_LOG(LogMovieRenderPipeline, VeryVerbose, TEXT("Failed to find cache data for Hitproxy! Id: %d"), HitProxyIndex);
+				}
+
+				OutCryptomatteBuffer[DstIndex] = Hash;
+			}
+		});
+	}
+
+	static void UpdateCryptomatteMetadata(const FObjectIdAccelerationData& InAccelData, const FString& InLayerName, TMap<FString, FString>& InOutMetadataMap)
+	{
+		const FString PassIdentifierHashAsShortString = InAccelData.PassIdentifierHashAsShortString;
+		InOutMetadataMap.Add(FString::Printf(TEXT("cryptomatte/%s/manifest"), *PassIdentifierHashAsShortString), InAccelData.JsonManifestCachedOutput);
+		InOutMetadataMap.Add(FString::Printf(TEXT("cryptomatte/%s/name"), *PassIdentifierHashAsShortString), InLayerName);
+		InOutMetadataMap.Add(FString::Printf(TEXT("cryptomatte/%s/hash"), *PassIdentifierHashAsShortString), TEXT("MurmurHash3_32"));
+		InOutMetadataMap.Add(FString::Printf(TEXT("cryptomatte/%s/conversion"), *PassIdentifierHashAsShortString), TEXT("uint32_to_float32"));
+	}
+}
diff --git a/Engine/Plugins/MovieScene/MoviePipelineMaskRenderPass/Source/MoviePipelineMaskRenderPass/Private/MovieRenderOverlappedMask.h b/Engine/Plugins/MovieScene/MoviePipelineMaskRenderPass/Source/MoviePipelineMaskRenderPass/Private/MovieRenderOverlappedMask.h
--- a/Engine/Plugins/MovieScene/MoviePipelineMaskRenderPass/Source/MoviePipelineMaskRenderPass/Private/MovieRenderOverlappedMask.h
+++ b/Engine/Plugins/MovieScene/MoviePipelineMaskRenderPass/Source/MoviePipelineMaskRenderPass/Private/MovieRenderOverlappedMask.h
@@ -35,6 +35,12 @@ public:
 	{
 	}
 
+	static FName GetName()
+	{
+		static FName Name = FName("MaskOverlappedAccumulator");
+		return Name;
+	}
+
 	/**
 	 * Allocates memory.
 	 *
diff --git a/Engine/Plugins/MovieScene/MovieRenderPipeline/Source/MovieRenderPipelineCore/Private/Graph/MovieGraphDefaultRenderer.cpp b/Engine/Plugins/MovieScene/MovieRenderPipeline/Source/MovieRenderPipelineCore/Private/Graph/MovieGraphDefaultRenderer.cpp
--- a/Engine/Plugins/MovieScene/MovieRenderPipeline/Source/MovieRenderPipelineCore/Private/Graph/MovieGraphDefaultRenderer.cpp
+++ b/Engine/Plugins/MovieScene/MovieRenderPipeline/Source/MovieRenderPipelineCore/Private/Graph/MovieGraphDefaultRenderer.cpp
@@ -432,7 +432,7 @@ TObjectPtr<UTextureRenderTarget2D> UMovieGraphDefaultRenderer::CreateViewRenderT
 	TObjectPtr<UTextureRenderTarget2D> NewTarget = NewObject<UTextureRenderTarget2D>(GetTransientPackage());
 	NewTarget->ClearColor = FLinearColor(0.0f, 0.0f, 0.0f, 0.0f);
 	NewTarget->TargetGamma = InInitParams.TargetGamma;
-	NewTarget->InitCustomFormat(InInitParams.Size.X, InInitParams.Size.Y, InInitParams.PixelFormat, false);
+	NewTarget->InitCustomFormat(InInitParams.Size.X, InInitParams.Size.Y, InInitParams.PixelFormat, InInitParams.bForceLinearGamma);
 	int32 ResourceSizeBytes = NewTarget->GetResourceSizeBytes(EResourceSizeMode::Type::EstimatedTotal);
 	UE_LOG(LogMovieRenderPipeline, Log, TEXT("Allocated a View Render Target sized: (%d, %d), Bytes: %d"), InInitParams.Size.X, InInitParams.Size.Y, ResourceSizeBytes);
 
diff --git a/Engine/Plugins/MovieScene/MovieRenderPipeline/Source/MovieRenderPipelineCore/Private/Graph/MovieGraphOutputMerger.cpp b/Engine/Plugins/MovieScene/MovieRenderPipeline/Source/MovieRenderPipelineCore/Private/Graph/MovieGraphOutputMerger.cpp
--- a/Engine/Plugins/MovieScene/MovieRenderPipeline/Source/MovieRenderPipelineCore/Private/Graph/MovieGraphOutputMerger.cpp
+++ b/Engine/Plugins/MovieScene/MovieRenderPipeline/Source/MovieRenderPipelineCore/Private/Graph/MovieGraphOutputMerger.cpp
@@ -94,6 +94,12 @@ namespace UE::MovieGraph
 		
 		if (FinishedPasses == TotalPasses)
 		{
+			// Merge in any metadata produced by the render
+			for (const TPair<FString, FString>& Metadata : Payload->AdditionalFileMetadata)
+			{
+				OutputFrame->FileMetadata.Add(Metadata);
+			}
+			
 			// Sort the output frames. This is only really important for multi-channel formats like EXR, but it lets passes
 			// specify which one should be the thumbnail/default rgba channels instead of a first-come-first-serve.
 			OutputFrame->ImageOutputData.ValueStableSort([](const TUniquePtr<FImagePixelData>& First, const TUniquePtr<FImagePixelData>& Second) -> bool
diff --git a/Engine/Plugins/MovieScene/MovieRenderPipeline/Source/MovieRenderPipelineCore/Public/Graph/MovieGraphDataTypes.h b/Engine/Plugins/MovieScene/MovieRenderPipeline/Source/MovieRenderPipelineCore/Public/Graph/MovieGraphDataTypes.h
--- a/Engine/Plugins/MovieScene/MovieRenderPipeline/Source/MovieRenderPipelineCore/Public/Graph/MovieGraphDataTypes.h
+++ b/Engine/Plugins/MovieScene/MovieRenderPipeline/Source/MovieRenderPipelineCore/Public/Graph/MovieGraphDataTypes.h
@@ -344,6 +344,11 @@ namespace UE::MovieGraph
 			, bAllowOCIO(true)
 		{}
 
+		virtual TSharedRef<FMovieGraphSampleState> Copy() const
+		{
+			return MakeShared<FMovieGraphSampleState>(*this);
+		}
+
 		/** The traversal context used to read graph values at the time of submission. */
 		FMovieGraphTraversalContext TraversalContext;
 
@@ -391,6 +396,19 @@ namespace UE::MovieGraph
 
 		/** Render scene capture source used for tracking the output color space (without OpenColorIO). */
 		ESceneCaptureSource SceneCaptureSource;
+
+		/**
+		 * Additional metadata that should be added to the output. In most cases, the node's GetFormatResolveArgs() should be used to provide
+		 * metadata unless the metadata is generated at a point in the pipeline where GetFormatResolveArgs() cannot be used.
+		 */
+		TMap<FString, FString> AdditionalFileMetadata;
+
+		/**
+		 * For multi-layer output formats, setting this allows the layer name to be explicitly specified. Normally the layer name will be
+		 * procedurally generated from multiple data sources to avoid collision with other layer names, but there are scenarios where specifying an
+		 * exact layer name may be needed.
+		 */
+		FString LayerNameOverride;
 	};
 
 	/**
diff --git a/Engine/Plugins/MovieScene/MovieRenderPipeline/Source/MovieRenderPipelineCore/Public/Graph/MovieGraphDefaultRenderer.h b/Engine/Plugins/MovieScene/MovieRenderPipeline/Source/MovieRenderPipelineCore/Public/Graph/MovieGraphDefaultRenderer.h
--- a/Engine/Plugins/MovieScene/MovieRenderPipeline/Source/MovieRenderPipelineCore/Public/Graph/MovieGraphDefaultRenderer.h
+++ b/Engine/Plugins/MovieScene/MovieRenderPipeline/Source/MovieRenderPipelineCore/Public/Graph/MovieGraphDefaultRenderer.h
@@ -84,17 +84,22 @@ namespace UE::MovieGraph::DefaultRenderer
 		FRenderTargetInitParams()
 			: Size(FIntPoint(0, 0))
 			, TargetGamma(0.f)
+			, bForceLinearGamma(false)
 			, PixelFormat(EPixelFormat::PF_Unknown)
 		{
 		}
 
 		FIntPoint Size;
 		float TargetGamma;
+		bool bForceLinearGamma;
 		EPixelFormat PixelFormat;
 
 		bool operator == (const FRenderTargetInitParams& InRHS) const
 		{
-			return Size == InRHS.Size && TargetGamma == InRHS.TargetGamma && PixelFormat == InRHS.PixelFormat;
+			return Size == InRHS.Size
+				&& TargetGamma == InRHS.TargetGamma
+				&& bForceLinearGamma == InRHS.bForceLinearGamma
+				&& PixelFormat == InRHS.PixelFormat;
 		}
 
 		bool operator != (const FRenderTargetInitParams& InRHS) const
@@ -104,7 +109,7 @@ namespace UE::MovieGraph::DefaultRenderer
 
 		friend uint32 GetTypeHash(FRenderTargetInitParams Params)
 		{
-			return HashCombineFast(GetTypeHash(Params.Size), HashCombineFast(GetTypeHash(Params.TargetGamma), GetTypeHash(Params.PixelFormat)));
+			return HashCombineFast(GetTypeHash(Params.Size), HashCombineFast(GetTypeHash(Params.TargetGamma), HashCombineFast(GetTypeHash(Params.bForceLinearGamma), GetTypeHash(Params.PixelFormat))));
 		}
 	};
 
diff --git a/Engine/Plugins/MovieScene/MovieRenderPipeline/Source/MovieRenderPipelineCore/Public/MovieRenderPipelineDataTypes.h b/Engine/Plugins/MovieScene/MovieRenderPipeline/Source/MovieRenderPipelineCore/Public/MovieRenderPipelineDataTypes.h
--- a/Engine/Plugins/MovieScene/MovieRenderPipeline/Source/MovieRenderPipelineCore/Public/MovieRenderPipelineDataTypes.h
+++ b/Engine/Plugins/MovieScene/MovieRenderPipeline/Source/MovieRenderPipelineCore/Public/MovieRenderPipelineDataTypes.h
@@ -1194,7 +1194,17 @@ public:
 
 namespace MoviePipeline
 {
-	struct MOVIERENDERPIPELINECORE_API IMoviePipelineOverlappedAccumulator : public TSharedFromThis<IMoviePipelineOverlappedAccumulator>
+	/** The interface that all accumulators should derive from. */
+	struct MOVIERENDERPIPELINECORE_API IMoviePipelineAccumulator
+	{
+	};
+
+	/** The interface that all accumulation args should derive from. */
+	struct MOVIERENDERPIPELINECORE_API IMoviePipelineAccumulationArgs
+	{
+	};
+	
+	struct MOVIERENDERPIPELINECORE_API IMoviePipelineOverlappedAccumulator : IMoviePipelineAccumulator, TSharedFromThis<IMoviePipelineOverlappedAccumulator>
 	{
 	};
 
diff --git a/Engine/Plugins/MovieScene/MovieRenderPipeline/Source/MovieRenderPipelineCore/Public/MovieRenderTileImage.h b/Engine/Plugins/MovieScene/MovieRenderPipeline/Source/MovieRenderPipelineCore/Public/MovieRenderTileImage.h
--- a/Engine/Plugins/MovieScene/MovieRenderPipeline/Source/MovieRenderPipelineCore/Public/MovieRenderTileImage.h
+++ b/Engine/Plugins/MovieScene/MovieRenderPipeline/Source/MovieRenderPipelineCore/Public/MovieRenderTileImage.h
@@ -5,6 +5,7 @@
 #include "CoreMinimal.h"
 
 #include "ImagePixelData.h"
+#include "MovieRenderPipelineDataTypes.h"
 
 /* Types
  *****************************************************************************/
@@ -72,7 +73,7 @@ public:
 /**
  * Contains all the image planes for the tiles.
  */
-struct MOVIERENDERPIPELINECORE_API FImageTileAccumulator
+struct MOVIERENDERPIPELINECORE_API FImageTileAccumulator : MoviePipeline::IMoviePipelineAccumulator
 {
 public:
 	/**
diff --git a/Engine/Plugins/MovieScene/MovieRenderPipeline/Source/MovieRenderPipelineEditor/Private/Customizations/Graph/MovieGraphShowFlagsCustomization.h b/Engine/Plugins/MovieScene/MovieRenderPipeline/Source/MovieRenderPipelineEditor/Private/Customizations/Graph/MovieGraphShowFlagsCustomization.h
--- a/Engine/Plugins/MovieScene/MovieRenderPipeline/Source/MovieRenderPipelineEditor/Private/Customizations/Graph/MovieGraphShowFlagsCustomization.h
+++ b/Engine/Plugins/MovieScene/MovieRenderPipeline/Source/MovieRenderPipelineEditor/Private/Customizations/Graph/MovieGraphShowFlagsCustomization.h
@@ -4,6 +4,7 @@
 
 #include "DetailWidgetRow.h"
 #include "EditorShowFlags.h"
+#include "Graph/Nodes/MovieGraphImagePassBaseNode.h"
 #include "Graph/Renderers/MovieGraphShowFlags.h"
 #include "IDetailChildrenBuilder.h"
 #include "IDetailGroup.h"
@@ -43,10 +44,19 @@ protected:
 			{SFG_Advanced, LOCTEXT("GraphAdvancedSF", "Advanced")}
 		};
 
+		check(PropertyHandle->IsValidHandle());
+
 		TArray<UObject*> OuterObjects;
 		PropertyHandle->GetOuterObjects(OuterObjects);
-		
-		check(PropertyHandle->IsValidHandle());
+
+		// Some renderer nodes don't allow Show Flags customization
+		if (UMovieGraphImagePassBaseNode* ImagePassBaseNode = Cast<UMovieGraphImagePassBaseNode>(OuterObjects[0]))
+		{
+			if (!ImagePassBaseNode->GetAllowsShowFlagsCustomization())
+			{
+				return;
+			}
+		}
 
 		const FObjectProperty* ObjectProperty = CastField<FObjectProperty>(PropertyHandle->GetProperty());
 		UMovieGraphShowFlags* ShowFlagObject = Cast<UMovieGraphShowFlags>(
diff --git a/Engine/Plugins/MovieScene/MovieRenderPipeline/Source/MovieRenderPipelineRenderPasses/Private/Graph/Renderers/MovieGraphDeferredPass.cpp b/Engine/Plugins/MovieScene/MovieRenderPipeline/Source/MovieRenderPipelineRenderPasses/Private/Graph/Renderers/MovieGraphDeferredPass.cpp
--- a/Engine/Plugins/MovieScene/MovieRenderPipeline/Source/MovieRenderPipelineRenderPasses/Private/Graph/Renderers/MovieGraphDeferredPass.cpp
+++ b/Engine/Plugins/MovieScene/MovieRenderPipeline/Source/MovieRenderPipelineRenderPasses/Private/Graph/Renderers/MovieGraphDeferredPass.cpp
@@ -146,24 +146,12 @@ void FMovieGraphDeferredPass::Render(const FMovieGraphTraversalContext& InFrameT
 	{
 		OverscanFraction = FMath::Clamp(CameraNode->OverscanPercentage / 100.f, 0.f, 1.f);
 	}
-	
-	FIntPoint AccumulatorResolution = UMovieGraphBlueprintLibrary::GetEffectiveOutputResolution(InTimeData.EvaluatedConfig);
+
 	// ToDo: When tiling is used, this should be the size of the per-tile backbuffer
+	FIntPoint AccumulatorResolution = UMovieGraphBlueprintLibrary::GetEffectiveOutputResolution(InTimeData.EvaluatedConfig);
 	FIntPoint BackbufferResolution = AccumulatorResolution;
-	// ToDo: This math probably needs the per-tile, pre-overlapped size? 
-	FIntPoint OverlappedPad = FIntPoint(FMath::CeilToInt(BackbufferResolution.X * TileOverlapPadRatio), FMath::CeilToInt(BackbufferResolution.Y * TileOverlapPadRatio));
-	// Calculate a backbuffer
-	UE::MovieGraph::DefaultRenderer::FRenderTargetInitParams RenderTargetInitParams;
-	{
-		RenderTargetInitParams.Size = BackbufferResolution;
-
-		// OCIO: Since this is a manually created Render target we don't need Gamma to be applied.
-		// We use this render target to render to via a display extension that utilizes Display Gamma
-		// which has a default value of 2.2 (DefaultDisplayGamma), therefore we need to set Gamma on this render target to 2.2 to cancel out any unwanted effects.
-		RenderTargetInitParams.TargetGamma = FOpenColorIORendering::DefaultDisplayGamma;
-		RenderTargetInitParams.PixelFormat = EPixelFormat::PF_FloatRGBA;
-	}
-
+	
+	DefaultRenderer::FRenderTargetInitParams RenderTargetInitParams = GetRenderTargetInitParams(InTimeData, AccumulatorResolution);
 	UTextureRenderTarget2D* RenderTarget = GraphRenderer->GetOrCreateViewRenderTarget(RenderTargetInitParams, RenderDataIdentifier);
 	FRenderTarget* RenderTargetResource = RenderTarget->GameThread_GetRenderTargetResource();
 	check(RenderTargetResource);
@@ -199,7 +187,6 @@ void FMovieGraphDeferredPass::Render(const FMovieGraphTraversalContext& InFrameT
 			CameraInfo.ViewActor = LocalPlayerController->GetViewTarget();
 		}
 
-
 		CameraInfo.bAllowCameraAspectRatio = true;
 		CameraInfo.TilingParams.TileSize = BackbufferResolution;
 		CameraInfo.TilingParams.OverlapPad = FVector2f(0.f, 0.f); // No tile support
@@ -211,7 +198,10 @@ void FMovieGraphDeferredPass::Render(const FMovieGraphTraversalContext& InFrameT
 		CameraInfo.SamplingParams.SpatialSampleCount = NumSpatialSamples;
 		CameraInfo.OverscanFraction = OverscanFraction;
 		CameraInfo.ProjectionMatrixJitterAmount = FVector2D((SpatialShiftAmount.X) * 2.0f / (float)BackbufferResolution.X, SpatialShiftAmount.Y * -2.0f / (float)BackbufferResolution.Y);
-
+		
+		// ToDo: This math probably needs the per-tile, pre-overlapped size? 
+		FIntPoint OverlappedPad = FIntPoint(FMath::CeilToInt(BackbufferResolution.X * TileOverlapPadRatio), FMath::CeilToInt(BackbufferResolution.Y * TileOverlapPadRatio));
+		
 		// For this particular tile, what is the offset into the output image
 		FIntPoint OverlappedOffset = FIntPoint(CameraInfo.TilingParams.TileIndexes.X * BackbufferResolution.X - OverlappedPad.X, CameraInfo.TilingParams.TileIndexes.Y * BackbufferResolution.Y - OverlappedPad.Y);
 		
diff --git a/Engine/Plugins/MovieScene/MovieRenderPipeline/Source/MovieRenderPipelineRenderPasses/Private/Graph/Renderers/MovieGraphImagePassBase.cpp b/Engine/Plugins/MovieScene/MovieRenderPipeline/Source/MovieRenderPipelineRenderPasses/Private/Graph/Renderers/MovieGraphImagePassBase.cpp
--- a/Engine/Plugins/MovieScene/MovieRenderPipeline/Source/MovieRenderPipelineRenderPasses/Private/Graph/Renderers/MovieGraphImagePassBase.cpp
+++ b/Engine/Plugins/MovieScene/MovieRenderPipeline/Source/MovieRenderPipelineRenderPasses/Private/Graph/Renderers/MovieGraphImagePassBase.cpp
@@ -9,6 +9,7 @@
 #include "MoviePipelineSurfaceReader.h"
 #include "CanvasTypes.h"
 #include "MovieRenderOverlappedImage.h"
+#include "OpenColorIORendering.h"
 #include "Engine/RendererSettings.h"
 #include "UnrealClient.h"
 #include "SceneViewExtensionContext.h"
@@ -80,6 +81,21 @@ FSceneView* FMovieGraphImagePassBase::CreateSceneView(const FSceneViewInitOption
 	return View;
 }
 
+DefaultRenderer::FRenderTargetInitParams FMovieGraphImagePassBase::GetRenderTargetInitParams(const FMovieGraphTimeStepData& InTimeData, const FIntPoint& InResolution)
+{
+	DefaultRenderer::FRenderTargetInitParams InitParams;
+
+	InitParams.Size = InResolution;
+
+	// OCIO: Since this is a manually created Render target we don't need Gamma to be applied.
+	// We use this render target to render to via a display extension that utilizes Display Gamma
+	// which has a default value of 2.2 (DefaultDisplayGamma), therefore we need to set Gamma on this render target to 2.2 to cancel out any unwanted effects.
+	InitParams.TargetGamma = FOpenColorIORendering::DefaultDisplayGamma;
+	InitParams.PixelFormat = PF_FloatRGBA;
+
+	return InitParams;
+}
+
 void FMovieGraphImagePassBase::ApplyCameraManagerPostProcessBlends(FSceneView* InView) const
 {
 	check(InView);
@@ -399,37 +415,31 @@ void FMovieGraphImagePassBase::PostRendererSubmission(
 		return;
 	}
 
-	FMoviePipelineAccumulatorPoolPtr SampleAccumulatorPool = GraphRenderer->GetOrCreateAccumulatorPool<FImageOverlappedAccumulator>();
-	UE::MovieGraph::DefaultRenderer::FSurfaceAccumulatorPool::FInstancePtr AccumulatorInstance = SampleAccumulatorPool->GetAccumulatorInstance_GameThread<FImageOverlappedAccumulator>(InSampleState.TraversalContext.Time.OutputFrameNumber, InSampleState.TraversalContext.RenderDataIdentifier);
-	
 	FMoviePipelineSurfaceQueuePtr LocalSurfaceQueue = GraphRenderer->GetOrCreateSurfaceQueue(InRenderTargetInitParams);
 	LocalSurfaceQueue->BlockUntilAnyAvailable();
 
-	FMovieGraphRenderDataAccumulationArgs AccumulationArgs;
-	{
-		AccumulationArgs.OutputMerger = GraphRenderer->GetOwningGraph()->GetOutputMerger();
-		AccumulationArgs.ImageAccumulator = StaticCastSharedPtr<FImageOverlappedAccumulator>(AccumulatorInstance->Accumulator);
-		AccumulationArgs.bIsFirstSample = InSampleState.TraversalContext.Time.bIsFirstTemporalSampleForFrame;
-		AccumulationArgs.bIsLastSample = InSampleState.TraversalContext.Time.bIsLastTemporalSampleForFrame;
-	}
+	TSharedRef<FMovieGraphRenderDataAccumulationArgs> AccumulationArgs =
+		StaticCastSharedRef<FMovieGraphRenderDataAccumulationArgs>(GetOrCreateAccumulator(GraphRenderer, InSampleState));
 
-	auto OnSurfaceReadbackFinished = [this, InSampleState, AccumulationArgs, AccumulatorInstance](TUniquePtr<FImagePixelData>&& InPixelData)
+	auto OnSurfaceReadbackFinished = [this, InSampleState, AccumulationArgs](TUniquePtr<FImagePixelData>&& InPixelData)
 	{
-		UE::Tasks::TTask<void> Task = UE::Tasks::Launch(UE_SOURCE_LOCATION, [PixelData = MoveTemp(InPixelData), InSampleState, AccumulationArgs, AccumulatorInstance]() mutable
+		FAccumulatorSampleFunc AccumulateFunction = GetAccumulateSampleFunction();
+		
+		const UE::Tasks::TTask<void> Task = UE::Tasks::Launch(UE_SOURCE_LOCATION, [PixelData = MoveTemp(InPixelData), InSampleState, AccumulationArgs, AccumulateFunction]() mutable
 		{
 			// Enqueue a encode for this frame onto our worker thread.
-			AccumulateSample_TaskThread(MoveTemp(PixelData), InSampleState, AccumulationArgs);
+			AccumulateFunction(MoveTemp(PixelData), InSampleState, AccumulationArgs);
 
 			// We have to defer clearing the accumulator until after sample accumulation has finished
-			if (AccumulationArgs.bIsLastSample)
+			if (AccumulationArgs->bIsLastSample)
 			{
 				// Final sample has now been executed, free the accumulator for reuse.
-				AccumulatorInstance->SetIsActive(false);
+				AccumulationArgs->AccumulatorInstance->SetIsActive(false);
 			}
-		}, AccumulatorInstance->TaskPrereq);
+		}, AccumulationArgs->AccumulatorInstance->TaskPrereq);
 
 		// Make the next accumulation task that uses this accumulator use the task we just created as a pre-req.
-		AccumulatorInstance->TaskPrereq = Task;
+		AccumulationArgs->AccumulatorInstance->TaskPrereq = Task;
 
 		// Because we're run on a separate thread, we need to check validity differently. The standard
 		// TWeakObjectPtr will report non-valid during GC (even if the object it's pointing to isn't being
@@ -469,24 +479,15 @@ TFunction<void(TUniquePtr<FImagePixelData>&&)> FMovieGraphImagePassBase::MakeFor
 	{
 		return nullptr;
 	}
-	
-	FMoviePipelineAccumulatorPoolPtr SampleAccumulator = GraphRenderer->GetOrCreateAccumulatorPool<FImageOverlappedAccumulator>();
-	UE::MovieGraph::DefaultRenderer::FSurfaceAccumulatorPool::FInstancePtr AccumulatorInstance =
-		SampleAccumulator->GetAccumulatorInstance_GameThread<FImageOverlappedAccumulator>(
-			InTimeData.RenderedFrameNumber, InSampleState.TraversalContext.RenderDataIdentifier);
-	
-	FMovieGraphRenderDataAccumulationArgs AccumulationArgs;
-	{
-		AccumulationArgs.OutputMerger = GraphRenderer->GetOwningGraph()->GetOutputMerger();
-		AccumulationArgs.ImageAccumulator = StaticCastSharedPtr<FImageOverlappedAccumulator>(AccumulatorInstance->Accumulator);
-		AccumulationArgs.bIsFirstSample = InTimeData.bIsFirstTemporalSampleForFrame;
-		AccumulationArgs.bIsLastSample = InTimeData.bIsLastTemporalSampleForFrame;
-	}
+		
+	// TODO: Is this correct? The prior code is *slightly* different.
+	TSharedRef<FMovieGraphRenderDataAccumulationArgs> AccumulationArgs =
+		StaticCastSharedRef<FMovieGraphRenderDataAccumulationArgs>(GetOrCreateAccumulator(GraphRenderer, InSampleState));
 
 	// The legacy surface reader takes the payload just so it can shuffle it into our callback, but we can just include the data
 	// directly in the callback, so this is just a dummy payload.
 	TSharedRef<FImagePixelDataPayload, ESPMode::ThreadSafe> FramePayload = MakeShared<FImagePixelDataPayload, ESPMode::ThreadSafe>();
-	auto Callback = [this, InSampleState, FramePayload, AccumulationArgs, AccumulatorInstance](TUniquePtr<FImagePixelData>&& InPixelData)
+	auto Callback = [this, InSampleState, FramePayload, AccumulationArgs](TUniquePtr<FImagePixelData>&& InPixelData)
 	{
 		// Transfer the framePayload to the returned data
 		TUniquePtr<FImagePixelData> PixelDataWithPayload = nullptr;
@@ -514,27 +515,49 @@ TFunction<void(TUniquePtr<FImagePixelData>&&)> FMovieGraphImagePassBase::MakeFor
 			checkNoEntry();
 		}
 
-		UE::Tasks::TTask<void> Task = UE::Tasks::Launch(UE_SOURCE_LOCATION, [PixelData = MoveTemp(PixelDataWithPayload), InSampleState, AccumulationArgs, AccumulatorInstance]() mutable
+		FAccumulatorSampleFunc AccumulateFunction = GetAccumulateSampleFunction();
+
+		UE::Tasks::TTask<void> Task = UE::Tasks::Launch(UE_SOURCE_LOCATION, [PixelData = MoveTemp(PixelDataWithPayload), InSampleState, AccumulationArgs, AccumulateFunction]() mutable
 		{
 			// Enqueue a encode for this frame onto our worker thread.
-			AccumulateSample_TaskThread(MoveTemp(PixelData), InSampleState, AccumulationArgs);
+			AccumulateFunction(MoveTemp(PixelData), InSampleState, AccumulationArgs);
 
 			// We have to defer clearing the accumulator until after sample accumulation has finished
-			if (AccumulationArgs.bIsLastSample)
+			if (AccumulationArgs->bIsLastSample)
 			{
 				// Final sample has now been executed, free the accumulator for reuse.
-				AccumulatorInstance->SetIsActive(false);
+				AccumulationArgs->AccumulatorInstance->SetIsActive(false);
 			}
-		}, AccumulatorInstance->TaskPrereq);
+		}, AccumulationArgs->AccumulatorInstance->TaskPrereq);
 
 		// Make the next accumulation task that uses this accumulator use the task we just created as a pre-req.
-		AccumulatorInstance->TaskPrereq = Task;
+		AccumulationArgs->AccumulatorInstance->TaskPrereq = Task;
 	};
 
 	return Callback;
 }
 
-void AccumulateSample_TaskThread(TUniquePtr<FImagePixelData>&& InPixelData, const ::UE::MovieGraph::FMovieGraphSampleState InSampleState, const FMovieGraphRenderDataAccumulationArgs& InAccumulationParams)
+TSharedRef<::MoviePipeline::IMoviePipelineAccumulationArgs> FMovieGraphImagePassBase::GetOrCreateAccumulator(const TObjectPtr<UMovieGraphDefaultRenderer> InGraphRenderer, const FMovieGraphSampleState& InSampleState) const
+{
+	const FMoviePipelineAccumulatorPoolPtr SampleAccumulatorPool = InGraphRenderer->GetOrCreateAccumulatorPool<FImageOverlappedAccumulator>();
+	const DefaultRenderer::FSurfaceAccumulatorPool::FInstancePtr AccumulatorInstance = SampleAccumulatorPool->GetAccumulatorInstance_GameThread<FImageOverlappedAccumulator>(InSampleState.TraversalContext.Time.OutputFrameNumber, InSampleState.TraversalContext.RenderDataIdentifier);
+
+	TSharedRef<FMovieGraphRenderDataAccumulationArgs> AccumulationArgs = MakeShared<FMovieGraphRenderDataAccumulationArgs>();
+	AccumulationArgs->OutputMerger = InGraphRenderer->GetOwningGraph()->GetOutputMerger();
+	AccumulationArgs->ImageAccumulator = StaticCastSharedPtr<FImageOverlappedAccumulator>(AccumulatorInstance->Accumulator);
+	AccumulationArgs->AccumulatorInstance = SampleAccumulatorPool->GetAccumulatorInstance_GameThread<FImageOverlappedAccumulator>(InSampleState.TraversalContext.Time.OutputFrameNumber, InSampleState.TraversalContext.RenderDataIdentifier);
+	AccumulationArgs->bIsFirstSample = InSampleState.TraversalContext.Time.bIsFirstTemporalSampleForFrame;
+	AccumulationArgs->bIsLastSample = InSampleState.TraversalContext.Time.bIsLastTemporalSampleForFrame;
+
+	return AccumulationArgs;
+}
+
+FMovieGraphImagePassBase::FAccumulatorSampleFunc FMovieGraphImagePassBase::GetAccumulateSampleFunction() const
+{
+	return AccumulateSample_TaskThread;
+}
+
+void AccumulateSample_TaskThread(TUniquePtr<FImagePixelData>&& InPixelData, const ::UE::MovieGraph::FMovieGraphSampleState InSampleState, const TSharedRef<::MoviePipeline::IMoviePipelineAccumulationArgs> InAccumulatorArgs)
 {
 	TRACE_CPUPROFILER_EVENT_SCOPE(MoviePipeline_AccumulateSample);
 		
@@ -545,7 +568,9 @@ void AccumulateSample_TaskThread(TUniquePtr<FImagePixelData>&& InPixelData, cons
 	TSharedPtr<FMovieGraphSampleState> SampleStatePayload = MakeShared<FMovieGraphSampleState>(InSampleState);
 	SamplePixelData->SetPayload(StaticCastSharedPtr<IImagePixelDataPayload>(SampleStatePayload));
 
-	TSharedPtr<IMovieGraphOutputMerger, ESPMode::ThreadSafe> OutputMergerPin = InAccumulationParams.OutputMerger.Pin();
+	const TSharedRef<FMovieGraphRenderDataAccumulationArgs, ESPMode::ThreadSafe> AccumulatorArgs = StaticCastSharedRef<FMovieGraphRenderDataAccumulationArgs>(InAccumulatorArgs);
+
+	const TSharedPtr<IMovieGraphOutputMerger, ESPMode::ThreadSafe> OutputMergerPin = AccumulatorArgs->OutputMerger.Pin();
 	if (!OutputMergerPin.IsValid())
 	{
 		return;
@@ -570,7 +595,7 @@ void AccumulateSample_TaskThread(TUniquePtr<FImagePixelData>&& InPixelData, cons
 		return;
 	}
 
-	TSharedPtr<FImageOverlappedAccumulator> AccumulatorPin = InAccumulationParams.ImageAccumulator.Pin();
+	const TSharedPtr<FImageOverlappedAccumulator> AccumulatorPin = StaticCastWeakPtr<FImageOverlappedAccumulator>(AccumulatorArgs->ImageAccumulator).Pin();
 	if (AccumulatorPin->NumChannels == 0)
 	{
 		LLM_SCOPE_BYNAME(TEXT("MoviePipeline/ImageAccumulatorInitMemory"));
diff --git a/Engine/Plugins/MovieScene/MovieRenderPipeline/Source/MovieRenderPipelineRenderPasses/Private/MovieGraphImageSequenceOutputNode.cpp b/Engine/Plugins/MovieScene/MovieRenderPipeline/Source/MovieRenderPipelineRenderPasses/Private/MovieGraphImageSequenceOutputNode.cpp
--- a/Engine/Plugins/MovieScene/MovieRenderPipeline/Source/MovieRenderPipelineRenderPasses/Private/MovieGraphImageSequenceOutputNode.cpp
+++ b/Engine/Plugins/MovieScene/MovieRenderPipeline/Source/MovieRenderPipelineRenderPasses/Private/MovieGraphImageSequenceOutputNode.cpp
@@ -457,9 +457,12 @@ void UMovieGraphImageSequenceOutputNode_MultiLayerEXR::OnReceiveImageDataImpl(UM
 			const UE::MovieGraph::FMovieGraphSampleState* Payload = ImageData->GetPayload<UE::MovieGraph::FMovieGraphSampleState>();
 			ShotIndex = Payload->TraversalContext.ShotIndex;
 
-			FString LayerName = {};
+			// The first layer doesn't get an explicit name. However, subsequent layers may have an explicit name specified for them. Normally layer
+			// names are procedurally generated though (below).
+			FString LayerName = (LayerIndex != 0) ? Payload->LayerNameOverride : FString();
 
-			if (LayerIndex != 0)
+			// Generate a procedural layer name if an explicit name wasn't specified.
+			if ((LayerIndex != 0) && LayerName.IsEmpty())
 			{
 				// If there is more than one layer, then we will prefix the layer. The first layer is not prefixed (and gets inserted as RGBA)
 				// as most programs that handle EXRs expect the main image data to be in an unnamed layer. We only postfix with cameraname
diff --git a/Engine/Plugins/MovieScene/MovieRenderPipeline/Source/MovieRenderPipelineRenderPasses/Public/Graph/Nodes/MovieGraphImagePassBaseNode.h b/Engine/Plugins/MovieScene/MovieRenderPipeline/Source/MovieRenderPipelineRenderPasses/Public/Graph/Nodes/MovieGraphImagePassBaseNode.h
--- a/Engine/Plugins/MovieScene/MovieRenderPipeline/Source/MovieRenderPipelineRenderPasses/Public/Graph/Nodes/MovieGraphImagePassBaseNode.h
+++ b/Engine/Plugins/MovieScene/MovieRenderPipeline/Source/MovieRenderPipelineRenderPasses/Public/Graph/Nodes/MovieGraphImagePassBaseNode.h
@@ -62,6 +62,9 @@ public:
 	/** Which AA Method should be used? */
 	virtual EAntiAliasingMethod GetAntiAliasingMethod() const { return EAntiAliasingMethod::AAM_None; }
 
+	/** Whether this node allows changing the Show Flags in the details panel. */
+	virtual bool GetAllowsShowFlagsCustomization() const { return true; }
+
 protected:
 	// Note: Since *individual* show flags are overridden instead of the entire ShowFlags property, manually set to
 	// overridden so the traversal picks the changes up (otherwise they will be ignored).
diff --git a/Engine/Plugins/MovieScene/MovieRenderPipeline/Source/MovieRenderPipelineRenderPasses/Public/Graph/Renderers/MovieGraphImagePassBase.h b/Engine/Plugins/MovieScene/MovieRenderPipeline/Source/MovieRenderPipelineRenderPasses/Public/Graph/Renderers/MovieGraphImagePassBase.h
--- a/Engine/Plugins/MovieScene/MovieRenderPipeline/Source/MovieRenderPipelineRenderPasses/Public/Graph/Renderers/MovieGraphImagePassBase.h
+++ b/Engine/Plugins/MovieScene/MovieRenderPipeline/Source/MovieRenderPipelineRenderPasses/Public/Graph/Renderers/MovieGraphImagePassBase.h
@@ -10,10 +10,11 @@ class UMovieGraphImagePassBaseNode;
 
 namespace UE::MovieGraph::Rendering
 {
-	struct MOVIERENDERPIPELINERENDERPASSES_API FMovieGraphRenderDataAccumulationArgs
+	struct MOVIERENDERPIPELINERENDERPASSES_API FMovieGraphRenderDataAccumulationArgs : ::MoviePipeline::IMoviePipelineAccumulationArgs
 	{
 	public:
-		TWeakPtr<FImageOverlappedAccumulator, ESPMode::ThreadSafe> ImageAccumulator;
+		TWeakPtr<::MoviePipeline::IMoviePipelineOverlappedAccumulator, ESPMode::ThreadSafe> ImageAccumulator;
+		DefaultRenderer::FSurfaceAccumulatorPool::FInstancePtr AccumulatorInstance;
 		TWeakPtr<IMovieGraphOutputMerger, ESPMode::ThreadSafe> OutputMerger;
 
 		// If it's the first sample then we will reset the accumulator to a clean slate before accumulating into it.
@@ -23,7 +24,7 @@ namespace UE::MovieGraph::Rendering
 	};
 
 	// Forward Declare
-	void AccumulateSample_TaskThread(TUniquePtr<FImagePixelData>&& InPixelData, const ::UE::MovieGraph::FMovieGraphSampleState InSampleState, const FMovieGraphRenderDataAccumulationArgs& InAccumulationParams);
+	void AccumulateSample_TaskThread(TUniquePtr<FImagePixelData>&& InPixelData, const FMovieGraphSampleState InSampleState, const TSharedRef<::MoviePipeline::IMoviePipelineAccumulationArgs> InAccumulatorArgs);
 
 	struct MOVIERENDERPIPELINERENDERPASSES_API FViewFamilyInitData
 	{
@@ -67,17 +68,26 @@ namespace UE::MovieGraph::Rendering
 		virtual TWeakObjectPtr<UMovieGraphDefaultRenderer> GetRenderer() const { return WeakGraphRenderer; }
 
 	protected:
+		typedef TFunction<void(TUniquePtr<FImagePixelData>&&, const FMovieGraphSampleState, const TSharedRef<::MoviePipeline::IMoviePipelineAccumulationArgs>)> FAccumulatorSampleFunc;
+		
 		/** Utility function for calculating a Projection Matrix (Orthographic or Perspective), and modify it based on the overscan percentage, aspect ratio, etc. */
 		virtual void CalculateProjectionMatrix(UE::MovieGraph::DefaultRenderer::FCameraInfo& InOutCameraInfo, FSceneViewProjectionData& InOutProjectionData, const FIntPoint InBackbufferResolution, const FIntPoint InAccumulatorResolution) const;
+		
 		/** Utility function for modifying the projection matrix to match the given tiling parameters. */
 		virtual void ModifyProjectionMatrixForTiling(const UE::MovieGraph::DefaultRenderer::FMovieGraphTilingParams& InTilingParams, const bool bInOrthographic, FMatrix& InOutProjectionMatrix, float& OutDoFSensorScale) const;
+		
 		/** Utility function for calculating the PrinciplePointOffset with the given tiling parameters. Used to make some effects (like vignette) work with tiling. */
 		virtual FVector4f CalculatePrinciplePointOffsetForTiling(const UE::MovieGraph::DefaultRenderer::FMovieGraphTilingParams& InTilingParams) const;
+		
 		/** Utility function for creating FSceneViewInitOptions based on the specified camera info. */
 		virtual FSceneViewInitOptions CreateViewInitOptions(const UE::MovieGraph::DefaultRenderer::FCameraInfo& InCameraInfo, FSceneViewFamilyContext* InViewFamily, FSceneViewStateReference& InViewStateRef) const;
+		
 		/** Utility function for creating a FSceneView for the given InitOptions, Family, and Camera. */
 		virtual FSceneView* CreateSceneView(const FSceneViewInitOptions& InInitOptions, TSharedRef<FSceneViewFamilyContext> InViewFamily, const UE::MovieGraph::DefaultRenderer::FCameraInfo& InCameraInfo) const;
 
+		/** Gets the render target init params that will be used when creating the render target. */
+		virtual DefaultRenderer::FRenderTargetInitParams GetRenderTargetInitParams(const FMovieGraphTimeStepData& InTimeData, const FIntPoint& InResolution);
+
 		virtual void ApplyCameraManagerPostProcessBlends(FSceneView* InView) const;
 		virtual TSharedRef<FSceneViewFamilyContext> CreateSceneViewFamily(const FViewFamilyInitData& InInitData) const;
 		virtual void ApplyMovieGraphOverridesToSceneView(TSharedRef<FSceneViewFamilyContext> InOutFamily, const FViewFamilyInitData& InInitData, const UE::MovieGraph::DefaultRenderer::FCameraInfo& InCameraInfo) const;
@@ -88,6 +98,12 @@ namespace UE::MovieGraph::Rendering
 		/** For this image pass, look up the associated node by type for the given config. */
 		virtual UMovieGraphImagePassBaseNode* GetParentNode(UMovieGraphEvaluatedConfig* InConfig) const { return nullptr; }
 
+		/** Gets an accumulator that's in use for the provided sample state (or creates a new one), and returns the accumulator args that are in use. */
+		virtual TSharedRef<::MoviePipeline::IMoviePipelineAccumulationArgs> GetOrCreateAccumulator(TObjectPtr<UMovieGraphDefaultRenderer> InGraphRenderer, const FMovieGraphSampleState& InSampleState) const;
+
+		/** Gets the function that the accumulator will use to perform accumulation. */
+		virtual FAccumulatorSampleFunc GetAccumulateSampleFunction() const;
+
 	protected:
 		TWeakObjectPtr<UMovieGraphDefaultRenderer> WeakGraphRenderer;
 	};
