MoviePipeline: Added missing resource transition in ObjectID and Panoramic passes.

diff --git a/Engine/Plugins/MovieScene/MoviePipelineMaskRenderPass/Source/MoviePipelineMaskRenderPass/Private/MoviePipelineObjectIdPass.cpp b/Engine/Plugins/MovieScene/MoviePipelineMaskRenderPass/Source/MoviePipelineMaskRenderPass/Private/MoviePipelineObjectIdPass.cpp
--- a/Engine/Plugins/MovieScene/MoviePipelineMaskRenderPass/Source/MoviePipelineMaskRenderPass/Private/MoviePipelineObjectIdPass.cpp
+++ b/Engine/Plugins/MovieScene/MoviePipelineMaskRenderPass/Source/MoviePipelineMaskRenderPass/Private/MoviePipelineObjectIdPass.cpp
@@ -236,6 +236,9 @@ void UMoviePipelineObjectIdRenderPass::RenderSample_GameThreadImpl(const FMovieP
 		ENQUEUE_RENDER_COMMAND(CanvasRenderTargetResolveCommand)(
 			[LocalSurfaceQueue, FramePayload, Callback, RenderTarget](FRHICommandListImmediate& RHICmdList) mutable
 			{
+				// Transition our render target from a render target view to a shader resource view to allow a shader to read from this Render Target.
+				RHICmdList.Transition(FRHITransitionInfo(RenderTarget->GetRenderTargetTexture(), ERHIAccess::RTV, ERHIAccess::SRVGraphicsPixel));
+
 				// Enqueue a encode for this frame onto our worker thread.
 				LocalSurfaceQueue->OnRenderTargetReady_RenderThread(RenderTarget->GetRenderTargetTexture(), FramePayload, MoveTemp(Callback));
 			});
diff --git a/Engine/Plugins/MovieScene/MoviePipelineMaskRenderPass/Source/MoviePipelineMaskRenderPass/Private/MoviePipelinePanoramicPass.cpp b/Engine/Plugins/MovieScene/MoviePipelineMaskRenderPass/Source/MoviePipelineMaskRenderPass/Private/MoviePipelinePanoramicPass.cpp
--- a/Engine/Plugins/MovieScene/MoviePipelineMaskRenderPass/Source/MoviePipelineMaskRenderPass/Private/MoviePipelinePanoramicPass.cpp
+++ b/Engine/Plugins/MovieScene/MoviePipelineMaskRenderPass/Source/MoviePipelineMaskRenderPass/Private/MoviePipelinePanoramicPass.cpp
@@ -554,6 +554,9 @@ void UMoviePipelinePanoramicPass::ScheduleReadbackAndAccumulation(const FMoviePi
 	ENQUEUE_RENDER_COMMAND(CanvasRenderTargetResolveCommand)(
 		[LocalSurfaceQueue, FramePayload, Callback, RenderTarget](FRHICommandListImmediate& RHICmdList) mutable
 		{
+			// Transition our render target from a render target view to a shader resource view to allow a shader to read from this Render Target.
+			RHICmdList.Transition(FRHITransitionInfo(RenderTarget->GetRenderTargetTexture(), ERHIAccess::RTV, ERHIAccess::SRVGraphicsPixel));
+
 			// Enqueue a encode for this frame onto our worker thread.
 			LocalSurfaceQueue->OnRenderTargetReady_RenderThread(RenderTarget->GetRenderTargetTexture(), FramePayload, MoveTemp(Callback));
 		});
