MoviePipeline: Fixed a GC-related crash that could occur when using the Object ID pass in the graph.

diff --git a/Engine/Plugins/MovieScene/MoviePipelineMaskRenderPass/Source/MoviePipelineMaskRenderPass/Private/Graph/MovieGraphObjectIdNode.cpp b/Engine/Plugins/MovieScene/MoviePipelineMaskRenderPass/Source/MoviePipelineMaskRenderPass/Private/Graph/MovieGraphObjectIdNode.cpp
--- a/Engine/Plugins/MovieScene/MoviePipelineMaskRenderPass/Source/MoviePipelineMaskRenderPass/Private/Graph/MovieGraphObjectIdNode.cpp
+++ b/Engine/Plugins/MovieScene/MoviePipelineMaskRenderPass/Source/MoviePipelineMaskRenderPass/Private/Graph/MovieGraphObjectIdNode.cpp
@@ -4,16 +4,8 @@
 
 #include "Editor/EditorPerProjectUserSettings.h"
 #include "MovieGraphObjectIdPass.h"
-#include "MoviePipelineObjectIdUtils.h"
 #include "MoviePipelineTelemetry.h"
 
-FUObjectAnnotationSparse<UE::MoviePipeline::FObjectIdAccelerationData, true>& UMovieGraphObjectIdNode::GetManifestAnnotation()
-{
-	static FUObjectAnnotationSparse<UE::MoviePipeline::FObjectIdAccelerationData, true> ManifestAnnotation;
-	
-	return ManifestAnnotation;
-}
-
 FEngineShowFlags UMovieGraphObjectIdNode::GetShowFlags() const
 {
 	FEngineShowFlags Flags(ESFIM_Game);
diff --git a/Engine/Plugins/MovieScene/MoviePipelineMaskRenderPass/Source/MoviePipelineMaskRenderPass/Private/Graph/MovieGraphObjectIdNode.h b/Engine/Plugins/MovieScene/MoviePipelineMaskRenderPass/Source/MoviePipelineMaskRenderPass/Private/Graph/MovieGraphObjectIdNode.h
--- a/Engine/Plugins/MovieScene/MoviePipelineMaskRenderPass/Source/MoviePipelineMaskRenderPass/Private/Graph/MovieGraphObjectIdNode.h
+++ b/Engine/Plugins/MovieScene/MoviePipelineMaskRenderPass/Source/MoviePipelineMaskRenderPass/Private/Graph/MovieGraphObjectIdNode.h
@@ -4,7 +4,6 @@
 
 #include "MoviePipelineObjectIdUtils.h"
 #include "Graph/Nodes/MovieGraphImagePassBaseNode.h"
-#include "UObject/UObjectAnnotation.h"
 
 #include "MovieGraphObjectIdNode.generated.h"
 
@@ -19,8 +18,6 @@ class UMovieGraphObjectIdNode : public UMovieGraphImagePassBaseNode
 public:
 	UMovieGraphObjectIdNode() = default;
 
-	static FUObjectAnnotationSparse<UE::MoviePipeline::FObjectIdAccelerationData, true>& GetManifestAnnotation();
-
 	// UMovieGraphImagePassBaseNode Interface
 	virtual FEngineShowFlags GetShowFlags() const override;
 	virtual EViewModeIndex GetViewModeIndex() const override;
diff --git a/Engine/Plugins/MovieScene/MoviePipelineMaskRenderPass/Source/MoviePipelineMaskRenderPass/Private/Graph/MovieGraphObjectIdPass.cpp b/Engine/Plugins/MovieScene/MoviePipelineMaskRenderPass/Source/MoviePipelineMaskRenderPass/Private/Graph/MovieGraphObjectIdPass.cpp
--- a/Engine/Plugins/MovieScene/MoviePipelineMaskRenderPass/Source/MoviePipelineMaskRenderPass/Private/Graph/MovieGraphObjectIdPass.cpp
+++ b/Engine/Plugins/MovieScene/MoviePipelineMaskRenderPass/Source/MoviePipelineMaskRenderPass/Private/Graph/MovieGraphObjectIdPass.cpp
@@ -104,11 +104,12 @@ namespace UE::MovieGraph
 
 			AccumulatorPin->FetchFinalPixelDataLinearColor(OutputLayers);
 
+			MoviePipeline::FObjectIdAccelerationData* AccelData = FMovieGraphObjectIdPass::GetAccelerationData(SampleStatePayload->TraversalContext.RenderDataIdentifier.RootBranchName);
+			check(AccelData);
+
 			// Add in the object ID metadata. This cannot be done in the node's GetFormatResolveArgs() because the manifest is only known after render-time,
-			// and the manifest data is destroyed during node teardown (and teardown occurs before the metadata is finalized and the file written to dis
-			FUObjectAnnotationSparse<UE::MoviePipeline::FObjectIdAccelerationData, true>& ManifestAnnotation = UMovieGraphObjectIdNode::GetManifestAnnotation();
-			const UE::MoviePipeline::FObjectIdAccelerationData AccelData = ManifestAnnotation.GetAnnotation(ObjectIdArgs->RenderPassNode.Get());
-			UpdateCryptomatteMetadata(AccelData, SampleStatePayload->TraversalContext.RenderDataIdentifier.RendererName, SampleStatePayload->AdditionalFileMetadata);
+			// and the manifest data is destroyed during node teardown (and teardown occurs before the metadata is finalized and the file written to disk)
+			UpdateCryptomatteMetadata(*AccelData, SampleStatePayload->TraversalContext.RenderDataIdentifier.RendererName, SampleStatePayload->AdditionalFileMetadata);
 
 			for (int32 Index = 0; Index < ObjectIdArgs->NumOutputLayers; Index++)
 			{
@@ -169,7 +170,7 @@ void FMovieGraphObjectIdPass::Setup(TWeakObjectPtr<UMovieGraphDefaultRenderer> I
 	// the Object ID render data identifiers are the same, other than the sub resource name, so this should be fine.
 	RenderDataIdentifier = RenderDataIdentifiers[0];
 
-	UE::MoviePipeline::FObjectIdAccelerationData AccelData = UE::MoviePipeline::FObjectIdAccelerationData();
+	UE::MoviePipeline::FObjectIdAccelerationData& AccelData = AccelerationDataByBranch.Add(LayerData.BranchName);
 
 	// Static metadata needed for Cryptomatte
 	const uint32 NameHash = ::MoviePipeline::HashNameToId(TCHAR_TO_UTF8(*InRenderPassNode->GetRendererName()));
@@ -193,8 +194,6 @@ void FMovieGraphObjectIdPass::Setup(TWeakObjectPtr<UMovieGraphDefaultRenderer> I
 		UpdateManifestAccelerationData(AccelData, ObjectIdNode->IdType);
 	}
 	
-	GetManifestAnnotation().AddAnnotation(InLayer.RenderPassNode.Get(), AccelData);
-
 	SceneViewState.Allocate(InRenderer->GetWorld()->GetFeatureLevel());
 
 	// The InRenderPassNode is not initialized with user's config. Use InLayer to 
@@ -207,7 +206,7 @@ void FMovieGraphObjectIdPass::Teardown()
 {
 	FMovieGraphDeferredPass::Teardown();
 
-	GetManifestAnnotation().RemoveAnnotation(LayerData.RenderPassNode.Get());
+	AccelerationDataByBranch.Remove(RenderDataIdentifier.RootBranchName);
 }
 
 void FMovieGraphObjectIdPass::GatherOutputPasses(UMovieGraphEvaluatedConfig* InConfig, TArray<FMovieGraphRenderDataIdentifier>& OutExpectedPasses) const
@@ -232,12 +231,15 @@ TSharedRef<MoviePipeline::IMoviePipelineAccumulationArgs> FMovieGraphObjectIdPas
 	const FMoviePipelineAccumulatorPoolPtr SampleAccumulatorPool = InGraphRenderer->GetOrCreateAccumulatorPool<FMaskOverlappedAccumulator>();
 	const UE::MovieGraph::DefaultRenderer::FSurfaceAccumulatorPool::FInstancePtr AccumulatorInstance = SampleAccumulatorPool->GetAccumulatorInstance_GameThread<FMaskOverlappedAccumulator>(InSampleState.TraversalContext.Time.OutputFrameNumber, InSampleState.TraversalContext.RenderDataIdentifier);
 
+	const UE::MoviePipeline::FObjectIdAccelerationData* AccelerationData = GetAccelerationData(InSampleState.TraversalContext.RenderDataIdentifier.RootBranchName);
+	check(AccelerationData);
+
 	TSharedRef<FMovieGraphObjectIdMaskSampleAccumulationArgs> AccumulationArgs = MakeShared<FMovieGraphObjectIdMaskSampleAccumulationArgs>();
 	AccumulationArgs->OutputMerger = InGraphRenderer->GetOwningGraph()->GetOutputMerger();
 	AccumulationArgs->ImageAccumulator = StaticCastSharedPtr<FMaskOverlappedAccumulator>(AccumulatorInstance->Accumulator);
 	AccumulationArgs->AccumulatorInstance = SampleAccumulatorPool->GetAccumulatorInstance_GameThread<FMaskOverlappedAccumulator>(InSampleState.TraversalContext.Time.OutputFrameNumber, InSampleState.TraversalContext.RenderDataIdentifier);
 	AccumulationArgs->NumOutputLayers = RenderDataIdentifiers.Num();
-	AccumulationArgs->CacheData = MakeShared<TMap<int32, UE::MoviePipeline::FMoviePipelineHitProxyCacheValue>>(*GetManifestAnnotation().GetAnnotation(LayerData.RenderPassNode.Get()).Cache);
+	AccumulationArgs->CacheData = MakeShared<TMap<int32, UE::MoviePipeline::FMoviePipelineHitProxyCacheValue>>(*AccelerationData->Cache);
 	AccumulationArgs->RenderPassNode = LayerData.RenderPassNode;
 
 	return AccumulationArgs;
@@ -248,6 +250,11 @@ UE::MovieGraph::Rendering::FMovieGraphImagePassBase::FAccumulatorSampleFunc FMov
 	return UE::MovieGraph::AccumulateSampleObjectId_TaskThread;
 }
 
+UE::MoviePipeline::FObjectIdAccelerationData* FMovieGraphObjectIdPass::GetAccelerationData(const FName& InBranchName)
+{
+	return AccelerationDataByBranch.Find(InBranchName);
+}
+
 UE::MovieGraph::DefaultRenderer::FRenderTargetInitParams FMovieGraphObjectIdPass::GetRenderTargetInitParams(const FMovieGraphTimeStepData& InTimeData, const FIntPoint& InResolution)
 {
 	UE::MovieGraph::DefaultRenderer::FRenderTargetInitParams InitParams;
@@ -260,8 +267,3 @@ UE::MovieGraph::DefaultRenderer::FRenderTargetInitParams FMovieGraphObjectIdPass
 
 	return InitParams;
 }
-
-FUObjectAnnotationSparse<UE::MoviePipeline::FObjectIdAccelerationData, true>& FMovieGraphObjectIdPass::GetManifestAnnotation()
-{
-	return UMovieGraphObjectIdNode::GetManifestAnnotation();
-}
\ No newline at end of file
diff --git a/Engine/Plugins/MovieScene/MoviePipelineMaskRenderPass/Source/MoviePipelineMaskRenderPass/Private/Graph/MovieGraphObjectIdPass.h b/Engine/Plugins/MovieScene/MoviePipelineMaskRenderPass/Source/MoviePipelineMaskRenderPass/Private/Graph/MovieGraphObjectIdPass.h
--- a/Engine/Plugins/MovieScene/MoviePipelineMaskRenderPass/Source/MoviePipelineMaskRenderPass/Private/Graph/MovieGraphObjectIdPass.h
+++ b/Engine/Plugins/MovieScene/MoviePipelineMaskRenderPass/Source/MoviePipelineMaskRenderPass/Private/Graph/MovieGraphObjectIdPass.h
@@ -5,7 +5,6 @@
 #include "MoviePipelineObjectIdUtils.h"
 #include "Graph/Renderers/MovieGraphDeferredPass.h"
 #include "Graph/Renderers/MovieGraphImagePassBase.h"
-#include "UObject/UObjectAnnotation.h"
 
 struct FMovieGraphObjectIdMaskSampleAccumulationArgs : MoviePipeline::IMoviePipelineAccumulationArgs
 {
@@ -39,13 +38,16 @@ struct FMovieGraphObjectIdPass : UE::MovieGraph::Rendering::FMovieGraphDeferredP
 	virtual FAccumulatorSampleFunc GetAccumulateSampleFunction() const override;
 	// ~FMovieGraphImagePassBase Interface
 
+	/** Gets the ObjectID acceleration data for a specific branch (will be nullptr if not found). */
+	static UE::MoviePipeline::FObjectIdAccelerationData* GetAccelerationData(const FName& InBranchName);
+
 protected:
 	virtual UE::MovieGraph::DefaultRenderer::FRenderTargetInitParams GetRenderTargetInitParams(const FMovieGraphTimeStepData& InTimeData, const FIntPoint& InResolution) override;
 
-private:
-	static FUObjectAnnotationSparse<UE::MoviePipeline::FObjectIdAccelerationData, true>& GetManifestAnnotation();
-
 protected:
 	/** The identifiers for all object ID layers that will be generated. */
 	TArray<FMovieGraphRenderDataIdentifier> RenderDataIdentifiers;
+
+	/** ObjectID acceleration data that needs to be consistent throughout a render (cached per branch). */
+	inline static TMap<FName, UE::MoviePipeline::FObjectIdAccelerationData> AccelerationDataByBranch;
 };
\ No newline at end of file
