MoviePipeline: Fixed multiple issues with Object ID when it is used across multiple layers (inconsistent IDs when multiple modifiers are in use, missing EXR metadata, flickering).

Also fixed:
* A performance/memory issue in the EXR writer that was especially noticeable when multiple Object ID layers were being generated in multi-layer EXRs.
* An issue with multi-layer EXRs that would cause layers with explicit naming to sometimes be written out to the primary/RGBA layer. This could break Object ID depending on the graph setup.
* Added multi-cam compatibility with Object ID. Previously Object ID only worked with the first camera.

#virtualized

diff --git a/Engine/Plugins/MovieScene/MoviePipelineMaskRenderPass/Source/MoviePipelineMaskRenderPass/Private/Graph/MovieGraphObjectIdPass.cpp b/Engine/Plugins/MovieScene/MoviePipelineMaskRenderPass/Source/MoviePipelineMaskRenderPass/Private/Graph/MovieGraphObjectIdPass.cpp
--- a/Engine/Plugins/MovieScene/MoviePipelineMaskRenderPass/Source/MoviePipelineMaskRenderPass/Private/Graph/MovieGraphObjectIdPass.cpp
+++ b/Engine/Plugins/MovieScene/MoviePipelineMaskRenderPass/Source/MoviePipelineMaskRenderPass/Private/Graph/MovieGraphObjectIdPass.cpp
@@ -5,12 +5,43 @@
 #include "MovieGraphObjectIdNode.h"
 #include "MoviePipelineHashUtils.h"
 #include "MoviePipelineObjectIdUtils.h"
+#include "MoviePipelineQueue.h"
 #include "Graph/MovieGraphBlueprintLibrary.h"
 #include "Graph/MovieGraphPipeline.h"
+#include "Graph/Nodes/MovieGraphCameraNode.h"
 #include "Graph/Nodes/MovieGraphImagePassBaseNode.h"
 
 namespace UE::MovieGraph
 {
+	/**
+	 * Gets the typename that will be used in the cryptomatte metadata ("typename" here is an official term in the cryptomatte spec). This is also
+	 * the name of the layer without the numerical suffix.
+	 */
+	static FString GetCryptomatteTypename(const FMovieGraphRenderDataIdentifier& InRenderDataIdentifier, const bool bIsMultiCam)
+	{
+		FString CryptomatteTypename = FString::Format(TEXT("{0}_{1}"), {InRenderDataIdentifier.RendererName, InRenderDataIdentifier.LayerName});
+
+		// Include the camera name if there is more than one camera being rendered.
+		if (bIsMultiCam)
+		{
+			CryptomatteTypename = FString::Format(TEXT("{0}_{1}"), {CryptomatteTypename, InRenderDataIdentifier.CameraName});
+		}
+
+		return CryptomatteTypename;
+	}
+
+	/** Gets the typename hash used as part of the cryptomatte metadata. Eg, "cryptomatte/<typename_hash>/..." */
+	static FString GetTypenameHash(const FMovieGraphRenderDataIdentifier& InRenderDataIdentifier)
+	{
+		// Note that the name hash includes the camera name as well because there may be multiple cameras rendered within one layer. Each layer/camera
+		// combination needs its own unique typename hash so it has a distinct entry in the metadata.
+		const uint32 NameHash = ::MoviePipeline::HashNameToId(TCHAR_TO_UTF8(*(InRenderDataIdentifier.LayerName + InRenderDataIdentifier.CameraName)));
+		FString TypenameHashString = FString::Printf(TEXT("%08x"), NameHash);
+		TypenameHashString.LeftInline(7);
+
+		return TypenameHashString;
+	}
+	
 	static void AccumulateSampleObjectId_TaskThread(TUniquePtr<FImagePixelData>&& InPixelData, const FMovieGraphSampleState InSampleState, const TSharedRef<::MoviePipeline::IMoviePipelineAccumulationArgs> InAccumulatorArgs)
 	{
 		TRACE_CPUPROFILER_EVENT_SCOPE(AccumulateSampleObjectId_TaskThread);
@@ -107,16 +138,26 @@ namespace UE::MovieGraph
 			MoviePipeline::FObjectIdAccelerationData* AccelData = FMovieGraphObjectIdPass::GetAccelerationData(SampleStatePayload->TraversalContext.RenderDataIdentifier.RootBranchName);
 			check(AccelData);
 
+			// If there are multiple cameras being rendered, layer names should include the camera name
+			bool bIsMultiCam = false;
+			{
+				const UMovieGraphCameraSettingNode* CameraNode = SampleStatePayload->TraversalContext.Time.EvaluatedConfig->GetSettingForBranch<UMovieGraphCameraSettingNode>(UMovieGraphNode::GlobalsPinName);
+				bIsMultiCam = (CameraNode && CameraNode->bRenderAllCameras) && (SampleStatePayload->TraversalContext.Shot->SidecarCameras.Num() > 1);
+			}
+
 			// Add in the object ID metadata. This cannot be done in the node's GetFormatResolveArgs() because the manifest is only known after render-time,
 			// and the manifest data is destroyed during node teardown (and teardown occurs before the metadata is finalized and the file written to disk)
-			UpdateCryptomatteMetadata(*AccelData, SampleStatePayload->TraversalContext.RenderDataIdentifier.RendererName, SampleStatePayload->AdditionalFileMetadata);
+			const FMovieGraphRenderDataIdentifier& RenderDataIdentifier = SampleStatePayload->TraversalContext.RenderDataIdentifier;
+			const FString CryptomatteTypename = GetCryptomatteTypename(RenderDataIdentifier, bIsMultiCam);
+			const FString TypenameHash = GetTypenameHash(RenderDataIdentifier);
+			UpdateCryptomatteMetadata(*AccelData, TypenameHash, CryptomatteTypename, SampleStatePayload->AdditionalFileMetadata);
 
 			for (int32 Index = 0; Index < ObjectIdArgs->NumOutputLayers; Index++)
 			{
 				// We unfortunately can't share ownership of the payload from the last sample due to the changed pass identifiers and layer name.
 				TSharedRef<FMovieGraphSampleState, ESPMode::ThreadSafe> NewPayload = SampleStatePayload->Copy();
 				FMovieGraphRenderDataIdentifier NewIdentifier = FMovieGraphRenderDataIdentifier(SampleStatePayload->TraversalContext.RenderDataIdentifier);
-				NewIdentifier.SubResourceName =  NewIdentifier.RendererName + FString::Printf(TEXT("%02d"), Index);
+				NewIdentifier.SubResourceName = CryptomatteTypename + FString::Printf(TEXT("%02d"), Index);
 				NewPayload->TraversalContext.RenderDataIdentifier = NewIdentifier;
 
 				// Update the layer name to be exactly what is needed for the Cryptomatte specification. The "name" in the metadata must match up
@@ -160,8 +201,8 @@ void FMovieGraphObjectIdPass::Setup(TWeakObjectPtr<UMovieGraphDefaultRenderer> I
 		NewIdentifier.RootBranchName = LayerData.BranchName;
 		NewIdentifier.LayerName = LayerData.LayerName;
 		NewIdentifier.RendererName = InRenderPassNode->GetRendererName();
-		NewIdentifier.SubResourceName = NewIdentifier.RendererName + FString::Printf(TEXT("%02d"), Index);
-		NewIdentifier.CameraName = InLayer.CameraName;
+		NewIdentifier.CameraName = LayerData.CameraName;
+		NewIdentifier.SubResourceName = UE::MovieGraph::GetCryptomatteTypename(NewIdentifier, InLayer.NumCameras > 1) + FString::Printf(TEXT("%02d"), Index);
 		
 		RenderDataIdentifiers.Add(NewIdentifier);
 	}
@@ -171,13 +212,6 @@ void FMovieGraphObjectIdPass::Setup(TWeakObjectPtr<UMovieGraphDefaultRenderer> I
 	RenderDataIdentifier = RenderDataIdentifiers[0];
 
 	UE::MoviePipeline::FObjectIdAccelerationData& AccelData = AccelerationDataByBranch.Add(LayerData.BranchName);
-
-	// Static metadata needed for Cryptomatte
-	const uint32 NameHash = ::MoviePipeline::HashNameToId(TCHAR_TO_UTF8(*InRenderPassNode->GetRendererName()));
-	FString PassIdentifierHashAsShortString = FString::Printf(TEXT("%08x"), NameHash);
-	PassIdentifierHashAsShortString.LeftInline(7);
-
-	AccelData.PassIdentifierHashAsShortString = PassIdentifierHashAsShortString;
 	AccelData.JsonManifest = MakeShared<FJsonObject>();
 	AccelData.Cache = MakeShared<TMap<int32, UE::MoviePipeline::FMoviePipelineHitProxyCacheValue>>();
 	AccelData.Cache->Reserve(1000);
@@ -226,6 +260,23 @@ UMovieGraphImagePassBaseNode* FMovieGraphObjectIdPass::GetParentNode(UMovieGraph
 	return ParentNode;
 }
 
+void FMovieGraphObjectIdPass::Render(const FMovieGraphTraversalContext& InFrameTraversalContext, const FMovieGraphTimeStepData& InTimeData)
+{
+	FMovieGraphDeferredPass::Render(InFrameTraversalContext, InTimeData);
+	
+	const UMovieGraphObjectIdNode* ObjectIdNode = Cast<UMovieGraphObjectIdNode>(LayerData.RenderPassNode);
+	check(ObjectIdNode);
+	
+	UE::MoviePipeline::FObjectIdAccelerationData* AccelData = GetAccelerationData(LayerData.BranchName);
+	check(AccelData);
+	
+	// This needs to be updated every frame. It's also important to call this AFTER the render finishes; HitProxy IDs will be out-of-date
+	// before the render starts. Render property changes made during the frame (via modifiers) will invalidate the HitProxy IDs in the global array,
+	// and those invalidations happen only once Render() is called (thus we need to wait until after Render() before updating the cache to
+	// get the new IDs).
+	UpdateManifestAccelerationData(*AccelData, ObjectIdNode->IdType);
+}
+
 TSharedRef<MoviePipeline::IMoviePipelineAccumulationArgs> FMovieGraphObjectIdPass::GetOrCreateAccumulator(TObjectPtr<UMovieGraphDefaultRenderer> InGraphRenderer, const UE::MovieGraph::FMovieGraphSampleState& InSampleState) const
 {
 	const FMoviePipelineAccumulatorPoolPtr SampleAccumulatorPool = InGraphRenderer->GetOrCreateAccumulatorPool<FMaskOverlappedAccumulator>();
diff --git a/Engine/Plugins/MovieScene/MoviePipelineMaskRenderPass/Source/MoviePipelineMaskRenderPass/Private/Graph/MovieGraphObjectIdPass.h b/Engine/Plugins/MovieScene/MoviePipelineMaskRenderPass/Source/MoviePipelineMaskRenderPass/Private/Graph/MovieGraphObjectIdPass.h
--- a/Engine/Plugins/MovieScene/MoviePipelineMaskRenderPass/Source/MoviePipelineMaskRenderPass/Private/Graph/MovieGraphObjectIdPass.h
+++ b/Engine/Plugins/MovieScene/MoviePipelineMaskRenderPass/Source/MoviePipelineMaskRenderPass/Private/Graph/MovieGraphObjectIdPass.h
@@ -38,6 +38,10 @@ struct FMovieGraphObjectIdPass : UE::MovieGraph::Rendering::FMovieGraphDeferredP
 	virtual FAccumulatorSampleFunc GetAccumulateSampleFunction() const override;
 	// ~FMovieGraphImagePassBase Interface
 
+	// FMovieGraphDeferredPass Interface
+	virtual void Render(const FMovieGraphTraversalContext& InFrameTraversalContext, const FMovieGraphTimeStepData& InTimeData) override;
+	// ~FMovieGraphDeferredPass Interface
+
 	/** Gets the ObjectID acceleration data for a specific branch (will be nullptr if not found). */
 	static UE::MoviePipeline::FObjectIdAccelerationData* GetAccelerationData(const FName& InBranchName);
 
diff --git a/Engine/Plugins/MovieScene/MoviePipelineMaskRenderPass/Source/MoviePipelineMaskRenderPass/Private/MoviePipelineObjectIdPass.cpp b/Engine/Plugins/MovieScene/MoviePipelineMaskRenderPass/Source/MoviePipelineMaskRenderPass/Private/MoviePipelineObjectIdPass.cpp
--- a/Engine/Plugins/MovieScene/MoviePipelineMaskRenderPass/Source/MoviePipelineMaskRenderPass/Private/MoviePipelineObjectIdPass.cpp
+++ b/Engine/Plugins/MovieScene/MoviePipelineMaskRenderPass/Source/MoviePipelineMaskRenderPass/Private/MoviePipelineObjectIdPass.cpp
@@ -97,17 +97,7 @@ void UMoviePipelineObjectIdRenderPass::SetupImpl(const MoviePipeline::FMoviePipe
 	AccumulatorPool = MakeShared<TAccumulatorPool<FMaskOverlappedAccumulator>, ESPMode::ThreadSafe>(6);
 
 	UE::MoviePipeline::FObjectIdAccelerationData AccelData = UE::MoviePipeline::FObjectIdAccelerationData();
-
-	// Static metadata needed for Cryptomatte
-	uint32 NameHash = MoviePipeline::HashNameToId(TCHAR_TO_UTF8(*PassIdentifier.Name));
-	FString PassIdentifierHashAsShortString = FString::Printf(TEXT("%08x"), NameHash);
-	PassIdentifierHashAsShortString.LeftInline(7);
-
-	AccelData.PassIdentifierHashAsShortString = PassIdentifierHashAsShortString;
-
-	
 	AccelData.JsonManifest = MakeShared<FJsonObject>();
-
 	AccelData.Cache = MakeShared<TMap<int32, UE::MoviePipeline::FMoviePipelineHitProxyCacheValue>>();
 	AccelData.Cache->Reserve(1000);
 
@@ -185,14 +175,13 @@ void UMoviePipelineObjectIdRenderPass::RenderSample_GameThreadImpl(const FMovieP
 		RHICmdList.Transition(FRHITransitionInfo(RenderTarget->GetRenderTargetTexture(), ERHIAccess::RTV, ERHIAccess::SRVGraphicsPixel));
 	});
 	
-
 	// Update the data in place, no need to copy back to the annotation.
 	UE::MoviePipeline::FObjectIdAccelerationData AccelData = ManifestAnnotation.GetAnnotation(this);
 	UE::MoviePipeline::UpdateManifestAccelerationData(AccelData, IdType);
 	ManifestAnnotation.AddAnnotation(this, AccelData);
 
 	// Update ObjectID-related file metadata
-	UpdateCryptomatteMetadata(AccelData, PassIdentifier.Name, InOutSampleState.OutputState.FileMetadata);
+	UpdateCryptomatteMetadata(AccelData, GetTypenameHash(), PassIdentifier.Name, InOutSampleState.OutputState.FileMetadata);
 
 	// Main Render Pass
 	{
@@ -254,6 +243,15 @@ void UMoviePipelineObjectIdRenderPass::PostRendererSubmission(const FMoviePipeli
 {
 }
 
+FString UMoviePipelineObjectIdRenderPass::GetTypenameHash() const
+{
+	const uint32 NameHash = MoviePipeline::HashNameToId(TCHAR_TO_UTF8(*PassIdentifier.Name));
+	FString TypenameHashString = FString::Printf(TEXT("%08x"), NameHash);
+	TypenameHashString.LeftInline(7);
+
+	return TypenameHashString;
+}
+
 namespace MoviePipeline
 {
 	static void AccumulateSampleObjectId_TaskThread(TUniquePtr<FImagePixelData>&& InPixelData, const UE::MoviePipeline::FObjectIdMaskSampleAccumulationArgs& InParams)
diff --git a/Engine/Plugins/MovieScene/MoviePipelineMaskRenderPass/Source/MoviePipelineMaskRenderPass/Private/MoviePipelineObjectIdPass.h b/Engine/Plugins/MovieScene/MoviePipelineMaskRenderPass/Source/MoviePipelineMaskRenderPass/Private/MoviePipelineObjectIdPass.h
--- a/Engine/Plugins/MovieScene/MoviePipelineMaskRenderPass/Source/MoviePipelineMaskRenderPass/Private/MoviePipelineObjectIdPass.h
+++ b/Engine/Plugins/MovieScene/MoviePipelineMaskRenderPass/Source/MoviePipelineMaskRenderPass/Private/MoviePipelineObjectIdPass.h
@@ -30,6 +30,10 @@ public:
 protected:
 	void PostRendererSubmission(const FMoviePipelineRenderPassMetrics& InSampleState);
 
+private:
+	/** Gets the typename hash used as part of the cryptomatte metadata. Eg, "cryptomatte/<typename_hash>/..." */
+	FString GetTypenameHash() const;
+
 public:
 	UPROPERTY(EditAnywhere, BlueprintReadWrite, Category = "Settings")
 	EMoviePipelineObjectIdPassIdType IdType;
diff --git a/Engine/Plugins/MovieScene/MoviePipelineMaskRenderPass/Source/MoviePipelineMaskRenderPass/Private/MoviePipelineObjectIdUtils.h b/Engine/Plugins/MovieScene/MoviePipelineMaskRenderPass/Source/MoviePipelineMaskRenderPass/Private/MoviePipelineObjectIdUtils.h
--- a/Engine/Plugins/MovieScene/MoviePipelineMaskRenderPass/Source/MoviePipelineMaskRenderPass/Private/MoviePipelineObjectIdUtils.h
+++ b/Engine/Plugins/MovieScene/MoviePipelineMaskRenderPass/Source/MoviePipelineMaskRenderPass/Private/MoviePipelineObjectIdUtils.h
@@ -89,8 +89,7 @@ namespace UE::MoviePipeline
 		{
 			return !Cache.IsValid()
 				&& !JsonManifest.IsValid()
-				&& JsonManifestCachedOutput.Len() == 0
-				&& PassIdentifierHashAsShortString.Len() == 0;
+				&& JsonManifestCachedOutput.Len() == 0;
 		}
 
 		/** Maps a HitProxy index to the data associated with the HitProxy. */
@@ -101,9 +100,6 @@ namespace UE::MoviePipeline
 
 		/** A cached and serialized version of the JsonManifest data. */
 		FString JsonManifestCachedOutput;
-
-		/** The pass identifier that is used within Cryptomatte metadata keys like `cryptomatte/<pass_id>/name`. */
-		FString PassIdentifierHashAsShortString;
 	};
 
 	struct FObjectIdMaskSampleAccumulationArgs : ::MoviePipeline::IMoviePipelineAccumulationArgs
@@ -349,12 +345,11 @@ namespace UE::MoviePipeline
 		});
 	}
 
-	static void UpdateCryptomatteMetadata(const FObjectIdAccelerationData& InAccelData, const FString& InLayerName, TMap<FString, FString>& InOutMetadataMap)
+	static void UpdateCryptomatteMetadata(const FObjectIdAccelerationData& InAccelData, const FString& InTypenameHash, const FString& InLayerName, TMap<FString, FString>& InOutMetadataMap)
 	{
-		const FString PassIdentifierHashAsShortString = InAccelData.PassIdentifierHashAsShortString;
-		InOutMetadataMap.Add(FString::Printf(TEXT("cryptomatte/%s/manifest"), *PassIdentifierHashAsShortString), InAccelData.JsonManifestCachedOutput);
-		InOutMetadataMap.Add(FString::Printf(TEXT("cryptomatte/%s/name"), *PassIdentifierHashAsShortString), InLayerName);
-		InOutMetadataMap.Add(FString::Printf(TEXT("cryptomatte/%s/hash"), *PassIdentifierHashAsShortString), TEXT("MurmurHash3_32"));
-		InOutMetadataMap.Add(FString::Printf(TEXT("cryptomatte/%s/conversion"), *PassIdentifierHashAsShortString), TEXT("uint32_to_float32"));
+		InOutMetadataMap.Add(FString::Printf(TEXT("cryptomatte/%s/manifest"), *InTypenameHash), InAccelData.JsonManifestCachedOutput);
+		InOutMetadataMap.Add(FString::Printf(TEXT("cryptomatte/%s/name"), *InTypenameHash), InLayerName);
+		InOutMetadataMap.Add(FString::Printf(TEXT("cryptomatte/%s/hash"), *InTypenameHash), TEXT("MurmurHash3_32"));
+		InOutMetadataMap.Add(FString::Printf(TEXT("cryptomatte/%s/conversion"), *InTypenameHash), TEXT("uint32_to_float32"));
 	}
 }
diff --git a/Engine/Plugins/MovieScene/MovieRenderPipeline/Source/MovieRenderPipelineCore/Private/Graph/MovieGraphDefaultRenderer.cpp b/Engine/Plugins/MovieScene/MovieRenderPipeline/Source/MovieRenderPipelineCore/Private/Graph/MovieGraphDefaultRenderer.cpp
--- a/Engine/Plugins/MovieScene/MovieRenderPipeline/Source/MovieRenderPipelineCore/Private/Graph/MovieGraphDefaultRenderer.cpp
+++ b/Engine/Plugins/MovieScene/MovieRenderPipeline/Source/MovieRenderPipelineCore/Private/Graph/MovieGraphDefaultRenderer.cpp
@@ -138,6 +138,7 @@ void UMovieGraphDefaultRenderer::SetupRenderingPipelineForShot(UMoviePipelineExe
 				LayerData.RenderPassNode = BranchRenderer.Value;
 				// Result could be -1 or 0...n
 				LayerData.CameraIndex = CameraIndexes[CameraIndex]; 
+				LayerData.NumCameras = CameraIndexes.Num();
 				
 				// We provide the name here so that the setup functions of the renderers can create RenderLayerIdentifiers.
 				UE::MovieGraph::DefaultRenderer::FCameraInfo CameraInfo = GetCameraInfo(LayerData.CameraIndex);
diff --git a/Engine/Plugins/MovieScene/MovieRenderPipeline/Source/MovieRenderPipelineCore/Private/Graph/MovieGraphOutputMerger.cpp b/Engine/Plugins/MovieScene/MovieRenderPipeline/Source/MovieRenderPipelineCore/Private/Graph/MovieGraphOutputMerger.cpp
--- a/Engine/Plugins/MovieScene/MovieRenderPipeline/Source/MovieRenderPipelineCore/Private/Graph/MovieGraphOutputMerger.cpp
+++ b/Engine/Plugins/MovieScene/MovieRenderPipeline/Source/MovieRenderPipelineCore/Private/Graph/MovieGraphOutputMerger.cpp
@@ -142,10 +142,13 @@ namespace UE::MovieGraph
 		
 		if (FinishedPasses == TotalPasses)
 		{
-			// Merge in any metadata produced by the render
-			for (const TPair<FString, FString>& Metadata : Payload->AdditionalFileMetadata)
+			// Merge in any metadata produced by the render (there may be multiple passes that need to be merged)
+			for (const FMovieGraphPassData& PassData : OutputFrame->ImageOutputData)
 			{
-				OutputFrame->FileMetadata.Add(Metadata);
+				const FMovieGraphSampleState* PassPayload = PassData.Value->GetPayload<FMovieGraphSampleState>();
+				check(PassPayload);
+
+				OutputFrame->FileMetadata.Append(PassPayload->AdditionalFileMetadata);
 			}
 			
 			// Sort the output frames. This is only really important for multi-channel formats like EXR, but it lets passes
diff --git a/Engine/Plugins/MovieScene/MovieRenderPipeline/Source/MovieRenderPipelineCore/Public/Graph/MovieGraphDataTypes.h b/Engine/Plugins/MovieScene/MovieRenderPipeline/Source/MovieRenderPipelineCore/Public/Graph/MovieGraphDataTypes.h
--- a/Engine/Plugins/MovieScene/MovieRenderPipeline/Source/MovieRenderPipelineCore/Public/Graph/MovieGraphDataTypes.h
+++ b/Engine/Plugins/MovieScene/MovieRenderPipeline/Source/MovieRenderPipelineCore/Public/Graph/MovieGraphDataTypes.h
@@ -319,6 +319,10 @@ struct FMovieGraphRenderPassLayerData
 	FString LayerName;
 	int32 CameraIndex;
 	FString CameraName;
+
+	// The number of cameras that are being rendered.
+	uint32 NumCameras = 1;
+	
 	TWeakObjectPtr<class UMovieGraphRenderPassNode> RenderPassNode;
 };
 // ToDo: Both of these can probably go into the Default Renderer implementation.
diff --git a/Engine/Plugins/MovieScene/MovieRenderPipeline/Source/MovieRenderPipelineRenderPasses/Private/MovieGraphImageSequenceOutputNode.cpp b/Engine/Plugins/MovieScene/MovieRenderPipeline/Source/MovieRenderPipelineRenderPasses/Private/MovieGraphImageSequenceOutputNode.cpp
--- a/Engine/Plugins/MovieScene/MovieRenderPipeline/Source/MovieRenderPipelineRenderPasses/Private/MovieGraphImageSequenceOutputNode.cpp
+++ b/Engine/Plugins/MovieScene/MovieRenderPipeline/Source/MovieRenderPipelineRenderPasses/Private/MovieGraphImageSequenceOutputNode.cpp
@@ -469,6 +469,7 @@ void UMovieGraphImageSequenceOutputNode_MultiLayerEXR::OnReceiveImageDataImpl(UM
 		PrepareTaskGlobalMetadata(*MultiLayerImageTask, InRawFrameData, FilenameToResolveArgs[Filename].FileMetadata);
 
 		// Add each render pass as a layer to the EXR
+		bool bHasGeneratedPrimaryRGBALayer = false;
 		int32 LayerIndex = 0;
 		int32 ShotIndex = 0;
 		for (const FMovieGraphRenderDataIdentifier& RenderID : RenderIDs)
@@ -479,12 +480,11 @@ void UMovieGraphImageSequenceOutputNode_MultiLayerEXR::OnReceiveImageDataImpl(UM
 			const UE::MovieGraph::FMovieGraphSampleState* Payload = ImageData->GetPayload<UE::MovieGraph::FMovieGraphSampleState>();
 			ShotIndex = Payload->TraversalContext.ShotIndex;
 
-			// The first layer doesn't get an explicit name. However, subsequent layers may have an explicit name specified for them. Normally layer
-			// names are procedurally generated though (below).
-			FString LayerName = (LayerIndex != 0) ? Payload->LayerNameOverride : FString();
-
-			// Generate a procedural layer name if an explicit name wasn't specified.
-			if ((LayerIndex != 0) && LayerName.IsEmpty())
+			// The first layer that doesn't have an explicitly-specified name will have an empty layer name in the EXR -- this is the "primary"/RGBA
+			// layer. Otherwise, the layer name will be procedurally generated if the payload didn't specify a layer name override. Having a "primary"
+			// layer in the EXR expands compatibility with a number of applications that read EXRs.
+			FString LayerName = Payload->LayerNameOverride;
+			if (bHasGeneratedPrimaryRGBALayer && LayerName.IsEmpty())
 			{
 				// If there is more than one layer, then we will prefix the layer. The first layer is not prefixed (and gets inserted as RGBA)
 				// as most programs that handle EXRs expect the main image data to be in an unnamed layer. We only postfix with cameraname
@@ -533,6 +533,11 @@ void UMovieGraphImageSequenceOutputNode_MultiLayerEXR::OnReceiveImageDataImpl(UM
 					}
 				}
 			}
+			else
+			{
+				// Don't generate a layer name. This layer will be the "primary" RGBA layer without a name.
+				bHasGeneratedPrimaryRGBALayer = true;
+			}
 
 			TUniquePtr<FImagePixelData> PixelData;
 			if (GetNumFileOutputNodes(*InRawFrameData->EvaluatedConfig, RenderID.RootBranchName) > 1)
diff --git a/Engine/Plugins/MovieScene/MovieRenderPipeline/Source/MovieRenderPipelineRenderPasses/Private/MoviePipelineEXROutput.cpp b/Engine/Plugins/MovieScene/MovieRenderPipeline/Source/MovieRenderPipelineRenderPasses/Private/MoviePipelineEXROutput.cpp
--- a/Engine/Plugins/MovieScene/MovieRenderPipeline/Source/MovieRenderPipelineRenderPasses/Private/MoviePipelineEXROutput.cpp
+++ b/Engine/Plugins/MovieScene/MovieRenderPipeline/Source/MovieRenderPipelineRenderPasses/Private/MoviePipelineEXROutput.cpp
@@ -233,7 +233,10 @@ bool FEXRImageWriteTask::WriteToDisk()
 			// To complete the file, EXR seeks back into the file and writes the scanline offsets when the file is closed, 
 			// which moves the tellp location. So file length is stored in advance for later use. The output file needs to be
 			// created after the header information is filled.
-			Imf::OutputFile ImfFile(OutputFile, Header, FPlatformMisc::NumberOfCoresIncludingHyperthreads());
+			//
+			// Note: OutputFile has an option to control the number of threads used to write the file. The default is fine;
+			// providing too many threads here will massively decrease performance and bloat memory usage.
+			Imf::OutputFile ImfFile(OutputFile, Header);
 #if WITH_EDITOR
 			try
 #endif
