MovieGraph: Added panoramic blending support to the Movie Graph system, mirroring functionality that Movie Render Queue had. Minor improvements to blending quality, memory usage, and blending performance (on high core count machines).

diff --git a/Engine/Plugins/MovieScene/MoviePipelineMaskRenderPass/Source/MoviePipelineMaskRenderPass/Private/Graph/MovieGraphObjectIdPass.cpp b/Engine/Plugins/MovieScene/MoviePipelineMaskRenderPass/Source/MoviePipelineMaskRenderPass/Private/Graph/MovieGraphObjectIdPass.cpp
--- a/Engine/Plugins/MovieScene/MoviePipelineMaskRenderPass/Source/MoviePipelineMaskRenderPass/Private/Graph/MovieGraphObjectIdPass.cpp
+++ b/Engine/Plugins/MovieScene/MoviePipelineMaskRenderPass/Source/MoviePipelineMaskRenderPass/Private/Graph/MovieGraphObjectIdPass.cpp
@@ -42,7 +42,7 @@ namespace UE::MovieGraph
 		return TypenameHashString;
 	}
 	
-	static void AccumulateSampleObjectId_TaskThread(TUniquePtr<FImagePixelData>&& InPixelData, const FMovieGraphSampleState InSampleState, const TSharedRef<::MoviePipeline::IMoviePipelineAccumulationArgs> InAccumulatorArgs)
+	static void AccumulateSampleObjectId_TaskThread(TUniquePtr<FImagePixelData>&& InPixelData, const TSharedRef<FMovieGraphSampleState> InSampleState, const TSharedRef<::MoviePipeline::IMoviePipelineAccumulationArgs> InAccumulatorArgs)
 	{
 		TRACE_CPUPROFILER_EVENT_SCOPE(AccumulateSampleObjectId_TaskThread);
 
@@ -50,7 +50,7 @@ namespace UE::MovieGraph
 
 		// Associate the sample state with the image as payload data, this allows downstream systems to fetch the values without us having to store the data
 		// separately and ensure they stay paired the whole way down.
-		TSharedPtr<FMovieGraphSampleState> SampleStatePayload = MakeShared<FMovieGraphSampleState>(InSampleState);
+		TSharedPtr<FMovieGraphSampleState> SampleStatePayload = InSampleState->Copy();
 		SamplePixelData->SetPayload(StaticCastSharedPtr<IImagePixelDataPayload>(SampleStatePayload));
 		
 		const TSharedRef<FMovieGraphObjectIdMaskSampleAccumulationArgs> ObjectIdArgs = StaticCastSharedRef<FMovieGraphObjectIdMaskSampleAccumulationArgs>(InAccumulatorArgs);
diff --git a/Engine/Plugins/MovieScene/MoviePipelineMaskRenderPass/Source/MoviePipelineMaskRenderPass/Private/MoviePipelinePanoramicBlender.cpp b/Engine/Plugins/MovieScene/MoviePipelineMaskRenderPass/Source/MoviePipelineMaskRenderPass/Private/MoviePipelinePanoramicBlender.cpp
--- a/Engine/Plugins/MovieScene/MoviePipelineMaskRenderPass/Source/MoviePipelineMaskRenderPass/Private/MoviePipelinePanoramicBlender.cpp
+++ b/Engine/Plugins/MovieScene/MoviePipelineMaskRenderPass/Source/MoviePipelineMaskRenderPass/Private/MoviePipelinePanoramicBlender.cpp
@@ -3,411 +3,142 @@
 #include "Math/PerspectiveMatrix.h"
 #include "MoviePipelinePanoramicPass.h"
 
-FMoviePipelinePanoramicBlender::FMoviePipelinePanoramicBlender(TSharedPtr<MoviePipeline::IMoviePipelineOutputMerger> InOutputMerger, const FIntPoint InOutputResolution)
+namespace UE::MoviePipeline
+{
+FMoviePipelinePanoramicBlender::FMoviePipelinePanoramicBlender(TSharedPtr<::MoviePipeline::IMoviePipelineOutputMerger> InOutputMerger, const FIntPoint InOutputResolution)
 	: OutputMerger(InOutputMerger)
 {
-	OutputEquirectangularMapSize = InOutputResolution;
+	OutputResolution = InOutputResolution;
 }
 
-static FLinearColor GetColorBilinearFiltered(const FImagePixelData* InSampleData, const FVector2D& InSamplePixelCoords, bool& OutClipped, bool bInForceAlphaToOpaque = false)
-{
-	// Pixel coordinates assume that 0.5, 0.5 is the center of the pixel, so we subtract half to make it indexable.
-	const FVector2D PixelCoordinateIndex = InSamplePixelCoords - 0.5f;
-
-	// Get surrounding pixels indices
-	FIntPoint LowerLeftPixelIndex = FIntPoint(FMath::FloorToInt(PixelCoordinateIndex.X), FMath::FloorToInt(PixelCoordinateIndex.Y));
-	FIntPoint LowerRightPixelIndex = FIntPoint(LowerLeftPixelIndex + FIntPoint(1, 0));
-	FIntPoint UpperLeftPixelIndex = FIntPoint(LowerLeftPixelIndex + FIntPoint(0, 1));
-	FIntPoint UpperRightPixelIndex = FIntPoint(LowerLeftPixelIndex + FIntPoint(1, 1));
-
-	// Clamp pixels indices to pixels array bounds.
-	// ToDo: Is this needed? Should we handle wrap around for the bottom right pixel? What gives
-	auto ClampPixelCoords = [&](FIntPoint& InOutPixelCoords, const FIntPoint& InArraySize)
-	{
-		if (InOutPixelCoords.X > InArraySize.X - 1 ||
-			InOutPixelCoords.Y > InArraySize.Y - 1 ||
-			InOutPixelCoords.X < 0 ||
-			InOutPixelCoords.Y < 0)
-		{
-			OutClipped = true;
-		}
-		InOutPixelCoords = FIntPoint(FMath::Clamp(InOutPixelCoords.X, 0, InArraySize.X - 1), FMath::Clamp(InOutPixelCoords.Y, 0, InArraySize.Y - 1));
-	};
-	ClampPixelCoords(LowerLeftPixelIndex, InSampleData->GetSize());
-	ClampPixelCoords(LowerRightPixelIndex, InSampleData->GetSize());
-	ClampPixelCoords(UpperLeftPixelIndex, InSampleData->GetSize());
-	ClampPixelCoords(UpperRightPixelIndex, InSampleData->GetSize());
-
-	// Fetch the colors for the four pixels. We convert to FLinearColor here so that our accumulation
-	// is done in linear space with enough precision. The samples are probably in F16 color right now.
-	FLinearColor LowerLeftPixelColor;
-	FLinearColor LowerRightPixelColor;
-	FLinearColor UpperLeftPixelColor;
-	FLinearColor UpperRightPixelColor;
-
-	int64 SizeInBytes = 0;
-	const void* SrcRawDataPtr = nullptr;
-	InSampleData->GetRawData(SrcRawDataPtr, SizeInBytes);
-
-	switch (InSampleData->GetType())
-	{
-	case EImagePixelType::Float16:
-	{
-		const FFloat16Color* ColorDataF16 = static_cast<const FFloat16Color*>(SrcRawDataPtr);
-		LowerLeftPixelColor = FLinearColor(ColorDataF16[LowerLeftPixelIndex.X + (LowerLeftPixelIndex.Y * InSampleData->GetSize().X)]);
-		LowerRightPixelColor = FLinearColor(ColorDataF16[LowerRightPixelIndex.X + (LowerRightPixelIndex.Y * InSampleData->GetSize().X)]);
-		UpperLeftPixelColor = FLinearColor(ColorDataF16[UpperLeftPixelIndex.X + (UpperLeftPixelIndex.Y * InSampleData->GetSize().X)]);
-		UpperRightPixelColor = FLinearColor(ColorDataF16[UpperRightPixelIndex.X + (UpperRightPixelIndex.Y * InSampleData->GetSize().X)]);
-	}
-	break;
-	case EImagePixelType::Float32:
-	{
-		const FLinearColor* ColorDataF32 = static_cast<const FLinearColor*>(SrcRawDataPtr);
-		LowerLeftPixelColor = ColorDataF32[LowerLeftPixelIndex.X + (LowerLeftPixelIndex.Y * InSampleData->GetSize().X)];
-		LowerRightPixelColor = ColorDataF32[LowerRightPixelIndex.X + (LowerRightPixelIndex.Y * InSampleData->GetSize().X)];
-		UpperLeftPixelColor = ColorDataF32[UpperLeftPixelIndex.X + (UpperLeftPixelIndex.Y * InSampleData->GetSize().X)];
-		UpperRightPixelColor = ColorDataF32[UpperRightPixelIndex.X + (UpperRightPixelIndex.Y * InSampleData->GetSize().X)];
-	}
-	break;
-	default:
-		// Not implemented
-		check(0);
-	}
-
-	// Interpolate between the 4 pixels based on the exact sub-pixel offset of the incoming coordinate (which may not be centered)
-	const float FracX = 1.f - FMath::Frac(InSamplePixelCoords.X);
-	const float FracY = 1.f - FMath::Frac(InSamplePixelCoords.Y);
-	FLinearColor InterpolatedPixelColor = 
-		 (LowerLeftPixelColor * (1.0f - FracX) + LowerRightPixelColor * FracX) * (1.0f - FracY)
-		 + (UpperLeftPixelColor * (1.0f - FracX) + UpperRightPixelColor * FracX) * FracY;
-
-	// Force final color alpha to opaque if requested
-	if (bInForceAlphaToOpaque)
-	{
-		InterpolatedPixelColor.A = 1.0f;
-	}
-
-	return InterpolatedPixelColor;
-}
-
-DECLARE_CYCLE_STAT(TEXT("STAT_MoviePipeline_PanoBlend"), STAT_MoviePipeline_PanoBlend, STATGROUP_MoviePipeline);
+DECLARE_CYCLE_STAT(TEXT("STAT_MoviePipeline_PanoBlendWait"), STAT_MoviePipeline_PanoBlendWait, STATGROUP_MoviePipeline);
 
 void FMoviePipelinePanoramicBlender::OnCompleteRenderPassDataAvailable_AnyThread(TUniquePtr<FImagePixelData>&& InData)
 {
-	SCOPE_CYCLE_COUNTER(STAT_MoviePipeline_PanoBlend);
+	// Copy the payload at the start, because it will get destroyed when the pixel data is moved to a task, which then
+	// uses it and discards it before this function finishes.
+	TSharedRef<FPanoramicImagePixelDataPayload> DataPayload = StaticCastSharedRef<FPanoramicImagePixelDataPayload>(InData->GetPayload<FPanoramicImagePixelDataPayload>()->Copy());
 
 	// This function is called every time a sample comes in from the GPU (after being accumulated) and needs to handle
 	// multiple samples from multiple frames being in flight at once. First step is to search to see if we're already
 	// working on an output frame for this sample.
-	TSharedPtr<FPanoramicOutputFrame> OutputFrame = nullptr;
-	TSharedPtr<FPanoramicBlendData> BlendDataTarget = nullptr;
-
-	FPanoramicImagePixelDataPayload* DataPayload = InData->GetPayload<FPanoramicImagePixelDataPayload>();
-	check(DataPayload);
-
-	const double BlendStartTime = FPlatformTime::Seconds();
-
-	// We calculate all of these immediately on the task thread because it's read-only and we need the
-	// information calculated by it to allocate memory when we start working on a sample.
-	FIntPoint SampleSize = DataPayload->Pane.Data.Resolution;
-	FRotator SampleRotation = FRotator(DataPayload->Pane.Data.bUseLocalRotation ? DataPayload->Pane.Data.CameraLocalRotation : DataPayload->Pane.Data.CameraRotation);
-	const float SampleHalfHorizontalFoVDegrees = 0.5f * DataPayload->Pane.Data.HorizontalFieldOfView;
-	const float SampleHalfVerticalFoVDegrees = 0.5f * DataPayload->Pane.Data.VerticalFieldOfView; // FMath::RadiansToDegrees((FMath::Atan((0.5f * SampleSize.X / FMath::Tan(DataPayload->Pane.Data.HorizontalFieldOfView / 2.f))))); // CTransformTools::getVFovFromHFovAndARInDegrees(iSampleFovInDegrees, static_cast<float>(iSampleSize.X) / static_cast<float>(iSampleSize.Y));
-	const float SampleHalfHorizontalFoVCosine = FMath::Cos(FMath::DegreesToRadians(SampleHalfHorizontalFoVDegrees));
-	const float SampleHalfVerticalFoVCosine = FMath::Cos(FMath::DegreesToRadians(SampleHalfVerticalFoVDegrees));
-
-	// Now calculate which direction the Panoramic Pane (that this sample represents) was facing originally.
-	const float SampleYawRad = FMath::DegreesToRadians(SampleRotation.Yaw);
-	const float SamplePitchRad = FMath::DegreesToRadians(SampleRotation.Pitch);
-	const FVector SampleDirectionOnTheta = FVector(FMath::Cos(SampleYawRad), FMath::Sin(SampleYawRad), 0);
-	const FVector SampleDirectionOnPhi = FVector(FMath::Cos(SamplePitchRad), 0.f, FMath::Sin(SamplePitchRad));
-
-	// Now construct a projection matrix representing the sample matching the original perspective it was taken from.
-	const FMatrix SampleProjectionMatrix = FReversedZPerspectiveMatrix(FMath::DegreesToRadians(SampleHalfHorizontalFoVDegrees), SampleSize.X, SampleSize.Y, DataPayload->Pane.Data.NearClippingPlane);
-
-	// For our given output size, figure out how many degrees each pixel represents.
-	const float EquiRectMapThetaStep = 360.f / (float)OutputEquirectangularMapSize.X;
-	const float EquiRectMapPhiStep = 180.f / (float)OutputEquirectangularMapSize.Y;
-
-	// Compute the index bounds in the equirectangular map corresponding to the sample bounds, so we don't loop over unnecessary pixels.
-	// This is approximated according to the weighting function for blending too.
-	// This assumes that the origin of the equirectangular map (0,0) has a yaw/pitch equal to -180/-90.
-	// Phi evolves in the opposite direction of Y (Y's origin is up-left)
-	// Pitch is clamped, because there is no vertical wrapping in the map
-	// Yaw is not clamped, because horizontal wrapping is possible. 
-	// The MinBound for X can actually be greater than the MaxBound due to wrapping, modulo is applied at eval time to ensure it wraps right.
-	float SampleYawMin = SampleRotation.Yaw - SampleHalfHorizontalFoVDegrees;
-	float SampleYawMax = SampleRotation.Yaw + SampleHalfHorizontalFoVDegrees;
-	int32 PixelIndexHorzMinBound = FMath::FloorToInt(((SampleYawMin) + 180.f) / EquiRectMapThetaStep);
-	int32 PixelIndexHorzMaxBound = FMath::FloorToInt(((SampleYawMax) + 180.f) / EquiRectMapThetaStep);
-
-	float SamplePitchMin = FMath::Max(SampleRotation.Pitch - SampleHalfVerticalFoVDegrees, -90.f); // Clamped to [-90, 90]
-	float SamplePitchMax = FMath::Min(SampleRotation.Pitch + SampleHalfVerticalFoVDegrees, 90.f); // Clamped to [-90, 90]
-	int32 PixelIndexVertMinBound = FMath::Max((OutputEquirectangularMapSize.Y) - FMath::FloorToInt((SamplePitchMax + 90.f) / EquiRectMapPhiStep), 0);
-	int32 PixelIndexVertMaxBound = FMath::Min((OutputEquirectangularMapSize.Y) - FMath::FloorToInt((SamplePitchMin + 90.f) / EquiRectMapPhiStep), OutputEquirectangularMapSize.Y);
-
+	FPoolEntry* TargetBlender = nullptr;
 	{
 		// Do a quick lock while we're iterating/adding to the PendingData array so a second sample
 		// doesn't come in mid iteration.
 		FScopeLock ScopeLock(&GlobalQueueDataMutex);
 
-		for (TPair<FMoviePipelineFrameOutputState, TSharedPtr<FPanoramicOutputFrame>>& KVP : PendingData)
+		for (TUniquePtr<FPoolEntry>& Item : PendingData)
 		{
-			if (KVP.Key.OutputFrameNumber == DataPayload->SampleState.OutputState.OutputFrameNumber)
+			if (Item->OutputFrameNumber == DataPayload->SampleState.OutputState.OutputFrameNumber && Item->bActive)
 			{
-				OutputFrame = KVP.Value;
+				TargetBlender = Item.Get();
 				break;
 			}
 		}
 
-		if (!OutputFrame)
+		if (!TargetBlender)
 		{
 			// UE_LOG(LogMovieRenderPipeline, Log, TEXT("Starting a new Output Frame in the Panoramic Blender for frame: %d"), DataPayload->SampleState.OutputState.OutputFrameNumber);
 
-			OutputFrame = PendingData.Add(DataPayload->SampleState.OutputState, MakeShared<FPanoramicOutputFrame>());
-			int32 EyeMultiplier = DataPayload->Pane.Data.EyeIndex == -1 ? 1 : 2;
-			int32 TotalSampleCount = DataPayload->Pane.Data.NumHorizontalSteps * DataPayload->Pane.Data.NumVerticalSteps * EyeMultiplier;
-			OutputFrame->NumSamplesTotal = TotalSampleCount;
-
+			// If we didn't find a blender already working on this frame, we'll try to re-use a previously allocated blender.
+			for (TUniquePtr<FPoolEntry>& Item : PendingData)
 			{
-				LLM_SCOPE_BYNAME(TEXT("MoviePipeline/PanoBlendFrameOutput"));
-				OutputFrame->OutputEquirectangularMap.SetNumZeroed(OutputEquirectangularMapSize.X * OutputEquirectangularMapSize.Y * EyeMultiplier);
-			}
-		}
-	}
-
-	// Now that we know which output frame we're contributing towards, we'll ask it for
-	// our own copy of the data so that we can blend without worrying about other threads.
-	check(OutputFrame);
-	{
-		FScopeLock ScopeLock(&GlobalQueueDataMutex);
-		BlendDataTarget = MakeShared<FPanoramicBlendData>();
-		BlendDataTarget->EyeIndex = DataPayload->Pane.Data.EyeIndex;
-		BlendDataTarget->bFinished = false;
-		BlendDataTarget->OriginalDataPayload = StaticCastSharedRef<FPanoramicImagePixelDataPayload>(DataPayload->Copy());
-		TArray<TSharedPtr<FPanoramicBlendData>>& EyeArray = OutputFrame->BlendedData.FindOrAdd(DataPayload->Pane.Data.EyeIndex);
-		EyeArray.Add(BlendDataTarget);
-	}
-
-	// Okay, our BlendDataTarget is now our own unique copy of the data to work on.
-	BlendDataTarget->BlendStartTime = FPlatformTime::Seconds();
-
-	// Build a rect that describes which part of the output map we'll be rendering into
-	BlendDataTarget->OutputBoundsMin = FIntPoint(PixelIndexHorzMinBound, PixelIndexVertMinBound);
-	BlendDataTarget->OutputBoundsMax = FIntPoint(PixelIndexHorzMaxBound, PixelIndexVertMaxBound);
-
-	BlendDataTarget->PixelWidth = BlendDataTarget->OutputBoundsMax.X - BlendDataTarget->OutputBoundsMin.X;
-	BlendDataTarget->PixelHeight = BlendDataTarget->OutputBoundsMax.Y - BlendDataTarget->OutputBoundsMin.Y;
-
-	// These need to be zeroed as we don't always touch every pixel in the rect with blending
-	// and they get +='d together later.
-	{
-		LLM_SCOPE_BYNAME(TEXT("MoviePipeline/PanoBlendPerTaskOutput"));
-		BlendDataTarget->Data.SetNumZeroed((BlendDataTarget->PixelWidth) * (BlendDataTarget->PixelHeight));
-	}
-
-	// Finally we can perform our actual blending. We blend into our intermediate buffer
-	// instead of the final output array to avoid multiple threads contending for pixels.
-	for (int32 Y = PixelIndexVertMinBound; Y < PixelIndexVertMaxBound; Y++)
-	{
-		for (int32 X = PixelIndexHorzMinBound; X < PixelIndexHorzMaxBound; X++)
-		{
-			// These X, Y coordinates are in output resolution space which is where we want to blend to.
-			// Our X bounds may go OOB, but we wrap horizontally so we need to figure out the proper X index.
-			const int32 OutputPixelX = ((X % OutputEquirectangularMapSize.X) + OutputEquirectangularMapSize.X) % OutputEquirectangularMapSize.X;
-			const int32 OutputPixelY = Y;
-
-			// Get the spherical coordinates (Theta and Phi) corresponding to the X and Y of the equirectangular map coordinates, converted to
-			// [-180, 180] and [-90, 90] coordinate space respectively. The half pixel offset is used to make the center of a pixel be considered
-			// that coordinate, and Phi increments in the opposite direction of Y.
-			const float Theta = EquiRectMapThetaStep * (((float)OutputPixelX) + 0.5f) - 180.f;
-			const float Phi = EquiRectMapPhiStep * (((float)OutputEquirectangularMapSize.Y - OutputPixelY) + 0.5f) - 90.f;
-
-			// Now convert the spherical coordinates into an actual direction (on the output map)
-			const float ThetaDeg = FMath::DegreesToRadians(Theta);
-			const float PhiDeg = FMath::DegreesToRadians(Phi);
-			const FVector OutputDirection(FMath::Cos(PhiDeg) * FMath::Cos(ThetaDeg), FMath::Cos(PhiDeg) * FMath::Sin(ThetaDeg), FMath::Sin(PhiDeg));
-			const FVector OutputDirectionTheta = FVector(FMath::Cos(ThetaDeg), FMath::Sin(ThetaDeg), 0);
-			const FVector OutputDirectionPhi = FVector(FMath::Cos(PhiDeg), 0.f, FMath::Sin(PhiDeg));
-
-			// Now we can compute how much the sample should influence this pixel. It is weighted by angular distance to the direction
-			// so that the edges have less influence (where they'd be more distorted anyways).
-			const float DirectionPhiDot = FVector::DotProduct(OutputDirectionPhi, SampleDirectionOnPhi); // ToDo: This only considers the whole Pano Pane and not per pixel of sample?
-			const float DirectionThetaDot = FVector::DotProduct(OutputDirectionTheta, SampleDirectionOnTheta);
-			const float WeightTheta = FMath::Max(DirectionThetaDot - SampleHalfHorizontalFoVCosine, 0.0f) / (1.0f - SampleHalfHorizontalFoVCosine);
-			const float WeightPhi = FMath::Max(DirectionPhiDot - SampleHalfVerticalFoVCosine, 0.0f) / (1.0f - SampleHalfVerticalFoVCosine);
-
-			const float SampleWeight = WeightTheta * WeightPhi;
-			const float SampleWeightSquared = SampleWeight* SampleWeight; // Exponential falloff produces a nicer blending result.
-
-			// The sample weight may be very small and not worth influencing this pixel.
-			if (SampleWeightSquared > KINDA_SMALL_NUMBER)
-			{
-				// Transform the direction vector from the equirectangular map world space to sample world space
-				FVector4 DirectionInSampleWorldSpace = FVector4(SampleRotation.UnrotateVector(OutputDirection), 1.0f);
-
-				static const FMatrix UnrealCoordinateConversion = FMatrix(
-					FPlane(0, 0, 1, 0),
-					FPlane(1, 0, 0, 0),
-					FPlane(0, 1, 0, 0),
-					FPlane(0, 0, 0, 1));
-				DirectionInSampleWorldSpace = UnrealCoordinateConversion.TransformFVector4(DirectionInSampleWorldSpace);
-
-				// Then project that direction into sample clip space
-				FVector4 DirectionInSampleClipSpace = SampleProjectionMatrix.TransformFVector4(DirectionInSampleWorldSpace);
-
-				// Converted into normalized device space (Divide by w for perspective)
-				FVector DirectionInSampleNDSpace = FVector(DirectionInSampleClipSpace) / DirectionInSampleClipSpace.W;
-
-				// Get the final pixel coordinates (direction in screen space)
-				FVector2D DirectionInSampleScreenSpace = ((FVector2D(DirectionInSampleNDSpace) + 1.0f) / 2.0f) * FVector2D(SampleSize.X, SampleSize.Y);
-
-				// Flip the Y value due to Y's zero coordinate being top left.
-				DirectionInSampleScreenSpace.Y = ((float)SampleSize.Y - DirectionInSampleScreenSpace.Y) - 1.0f;
-
-				// Do a bilinear color sample at the pixel coordinates (from the sample), weight it, and add it to the output map. We store
-				// weights separately so that we can preserve the alpha channel of the main image.
-				bool bClipped = false;
-				FLinearColor SampleColor = GetColorBilinearFiltered(InData.Get(), DirectionInSampleScreenSpace, bClipped, true);
-
-				if (!bClipped)
+				if (!Item->bActive)
 				{
-					// When we calculate the actual output location we need to shift the X/Y. This is because up until now the math has been done in
-					// output resolution space, but each sample only allocates a color map big enough for itself. It'll get shifted back out to the
-					// right location later.
-					int32 SampleOutputX = OutputPixelX - BlendDataTarget->OutputBoundsMin.X;
-					// Mod this again by our output map so we don't OOB on it. It'll wrap weirdly in the output map but should restore fine.
-					SampleOutputX = ((SampleOutputX % (BlendDataTarget->PixelWidth)) + (BlendDataTarget->PixelWidth)) % (BlendDataTarget->PixelWidth); // Positive Mod
-					int32 SampleOutputY = Y;
-					SampleOutputY -= BlendDataTarget->OutputBoundsMin.Y;
-						
-					const int32 FinalIndex = SampleOutputX + (SampleOutputY * (BlendDataTarget->PixelWidth));
-					BlendDataTarget->Data[FinalIndex] += SampleColor * SampleWeightSquared;
-					// OutputEquirectangularMap[OutputPixelX + (OutputPixelY * OutputEquirectangularMapSize.X)] += SampleColor * SampleWeightSquared; // ToDo move to weight map.
+					TargetBlender = Item.Get();
 				}
 			}
-		}
-	}
-
-	BlendDataTarget->BlendEndTime = FPlatformTime::Seconds();
 
-	// Blend the new sample into the output map as soon as possible, so that we can free the temporary memory held
-	// by this sample. This part is single-threaded (with respect to the other tasks).
-	{
-		// Lock access to our output map
-		FScopeLock ScopeLock(&OutputDataMutex);
+			// If we still don't have a target blender, then this is a new one and we need to allocate an entry.
+			if (!TargetBlender)
+			{
+				TUniquePtr<FPoolEntry> NewEntry = MakeUnique<FPoolEntry>();
+				int32 NewIndex = PendingData.Add(MoveTemp(NewEntry));
+				TargetBlender = PendingData[NewIndex].Get();
+			}
 
+			check(TargetBlender);
 
-		// For Stereo images, just offset the second eye by the entire output size so it goes on the bottom half
-		// of the image. We've already made the array twice as big as it needs to be (in the case of stereo) and
-		// below EyeIndex 0 (left eye) will go in the first half of the image, eye index 1 in the second.
-		int32 EyeOffset = 0;
-		if (BlendDataTarget->OriginalDataPayload->Pane.Data.EyeIndex != -1)
-		{
-			EyeOffset = (OutputEquirectangularMapSize.X * OutputEquirectangularMapSize.Y) * BlendDataTarget->OriginalDataPayload->Pane.Data.EyeIndex;
+			// If we were already working on this frame, the first for loop through the pending data would have found it.
+			// So we know that if we get here, that we need to initialize whatever blender we ended up with.
+			TargetBlender->OutputFrameNumber = DataPayload->SampleState.OutputState.OutputFrameNumber;
+			TargetBlender->bActive = true;
+			TargetBlender->NumCompletedAccumulations = 0;
+			TargetBlender->Blender.Initialize(OutputResolution);
 		}
+	}
 
-		for (int32 SampleY = 0; SampleY < BlendDataTarget->PixelHeight; SampleY++)
+	// This can get called later (due to blending being async) so only capture by value.
+	auto OnDebugSampleAvailable = [
+		DataPayloadCopy = DataPayload->Copy(),
+		WeakOutputMerger = OutputMerger](FLinearColor* Data, FIntPoint Resolution)
 		{
-			for (int32 SampleX = 0; SampleX < BlendDataTarget->PixelWidth; SampleX++)
+			TSharedRef<FPanoramicImagePixelDataPayload> PayloadAsPano = StaticCastSharedRef<FPanoramicImagePixelDataPayload>(DataPayloadCopy);
+			if (!PayloadAsPano->SampleState.bWriteSampleToDisk)
 			{
-				int32 OriginalX = SampleX + BlendDataTarget->OutputBoundsMin.X;
-				int32 OriginalY = SampleY + BlendDataTarget->OutputBoundsMin.Y;
-				const int32 OutputPixelX = ((OriginalX % OutputEquirectangularMapSize.X) + OutputEquirectangularMapSize.X) % OutputEquirectangularMapSize.X;
-				const int32 OutputPixelY = OriginalY;
-
-				int32 SourceIndex = SampleX + (SampleY * (BlendDataTarget->PixelWidth));
-				int32 DestIndex = OutputPixelX + (OutputPixelY * OutputEquirectangularMapSize.X);
-				OutputFrame->OutputEquirectangularMap[DestIndex + EyeOffset] += BlendDataTarget->Data[SourceIndex];
+				return;
 			}
-		}
-
-		bool bDebugSamples = DataPayload->SampleState.bWriteSampleToDisk;
-		if (bDebugSamples)
-		{
-			// Write each blended sample to the output as a debug sample so we can inspect the job blending is doing for each pane.
-			// Hack up the debug output name a bit so they're unique.
-			if (BlendDataTarget->OriginalDataPayload->Pane.Data.EyeIndex >= 0)
+			if (PayloadAsPano->Pane.Data.EyeIndex >= 0)
 			{
-				BlendDataTarget->OriginalDataPayload->Debug_OverrideFilename = FString::Printf(TEXT("/%s_PaneX_%d_PaneY_%dEye_%d-Blended.%d"),
-					*BlendDataTarget->OriginalDataPayload->PassIdentifier.Name, BlendDataTarget->OriginalDataPayload->Pane.Data.HorizontalStepIndex,
-					BlendDataTarget->OriginalDataPayload->Pane.Data.VerticalStepIndex, DataPayload->Pane.Data.EyeIndex, BlendDataTarget->OriginalDataPayload->SampleState.OutputState.OutputFrameNumber);
+				PayloadAsPano->Debug_OverrideFilename = FString::Printf(TEXT("/%s_PaneX_%d_PaneY_%dEye_%d-Blended.%d"),
+					*PayloadAsPano->PassIdentifier.Name, PayloadAsPano->Pane.Data.HorizontalStepIndex,
+					PayloadAsPano->Pane.Data.VerticalStepIndex, PayloadAsPano->Pane.Data.EyeIndex, PayloadAsPano->SampleState.OutputState.OutputFrameNumber);
 			}
 			else
 			{
-				BlendDataTarget->OriginalDataPayload->Debug_OverrideFilename = FString::Printf(TEXT("/%s_PaneX_%d_PaneY_%d-Blended.%d"),
-					*BlendDataTarget->OriginalDataPayload->PassIdentifier.Name, BlendDataTarget->OriginalDataPayload->Pane.Data.HorizontalStepIndex,
-					BlendDataTarget->OriginalDataPayload->Pane.Data.VerticalStepIndex, BlendDataTarget->OriginalDataPayload->SampleState.OutputState.OutputFrameNumber);
+				PayloadAsPano->Debug_OverrideFilename = FString::Printf(TEXT("/%s_PaneX_%d_PaneY_%d-Blended.%d"),
+					*PayloadAsPano->PassIdentifier.Name, PayloadAsPano->Pane.Data.HorizontalStepIndex,
+					PayloadAsPano->Pane.Data.VerticalStepIndex, PayloadAsPano->SampleState.OutputState.OutputFrameNumber);
 			}
 
-			// Now that the sample has been blended pass it (and the memory it owned, we already read from it) to the debug output step.
-			TUniquePtr<TImagePixelData<FLinearColor>> FinalPixelData = MakeUnique<TImagePixelData<FLinearColor>>(FIntPoint(BlendDataTarget->PixelWidth, BlendDataTarget->PixelHeight), TArray64<FLinearColor>(MoveTemp(BlendDataTarget->Data)), BlendDataTarget->OriginalDataPayload);
-			ensure(OutputMerger.IsValid());
-			OutputMerger.Pin()->OnSingleSampleDataAvailable_AnyThread(MoveTemp(FinalPixelData));
-		}
-		else
-		{
-			// Ensure we reset the memory we allocated, to minimize concurrent allocations.
-			BlendDataTarget->Data.Reset();
-		}
-	}
-
-	// This should only be set to true when this thread is really finished with the work.
-	bool bIsLastSample = false;
-	{
-		FScopeLock ScopeLock(&GlobalQueueDataMutex);
-		BlendDataTarget->bFinished = true;
+			// We have to copy the memory because the blender is going to re-use it.
+			TArray64<FLinearColor> BlendDataCopy = TArray64<FLinearColor>(Data, Resolution.X * Resolution.Y);
+			TUniquePtr<TImagePixelData<FLinearColor>> FinalPixelData = MakeUnique<TImagePixelData<FLinearColor>>(Resolution, MoveTemp(BlendDataCopy), PayloadAsPano->Copy());
 
-		// Check to see if all of the samples have come in from the GPU and have been blended
-		int32 NumFinishedSamples = 0;
-
-		// For each eye
-		for(TPair<int32, TArray<TSharedPtr<FPanoramicBlendData>>>& KVP : OutputFrame->BlendedData)
-		{
-			// For each sample in the eye
-			for (int32 Index = 0; Index < KVP.Value.Num(); Index++)
+			if (ensure(WeakOutputMerger.IsValid()))
 			{
-				NumFinishedSamples += KVP.Value[Index]->bFinished ? 1 : 0;
+				WeakOutputMerger.Pin()->OnSingleSampleDataAvailable_AnyThread(MoveTemp(FinalPixelData));
 			}
-		}
+		};
 
-		bIsLastSample = NumFinishedSamples == OutputFrame->NumSamplesTotal;
-	}
+	// Now that we know which blender we're trying to accumulate to, we can just send the data to it directly. We're already
+	// on a task thread, and the blending process supports multiple task threads working on blending at the same time.
+	TargetBlender->Blender.BlendSample_AnyThread(MoveTemp(InData), DataPayload->Pane.Data, OnDebugSampleAvailable);
+
+	// Checking to see if this is the last sample is slightly complicated, because we can have multiple threads in this function at the same time.
+	// Inside the blender, it only lets one thread increment the sample count at the same time, but that means when we look at it, we need to
+	// go through the same lock so that we don't have two threads (in this function) read the value, and both decide they're the last sample.
 
-	if (bIsLastSample)
 	{
+		// We put this behind our lock just so that we don't have two threads get the same value one after the other and still decide
+		// they're the last.
+		FScopeLock ScopeLock(&GlobalQueueDataMutex);
+		TargetBlender->NumCompletedAccumulations++;
+
+		int32 NumCompletedAccumulations = TargetBlender->NumCompletedAccumulations;
+		const bool bLastSample = NumCompletedAccumulations == DataPayload->Pane.Data.NumHorizontalSteps * DataPayload->Pane.Data.NumVerticalSteps;
+
+		if (bLastSample)
 		{
-			// Now that we've accumulated all the pixel values, we need to normalize them.
-			for (int32 PixelIndex = 0; PixelIndex < OutputFrame->OutputEquirectangularMap.Num(); PixelIndex++)
+			// BlendSample_AnyThread returns immediately and we'll increment it as completed, so if this is the last sample,
+			// we'll wait for the outstanding work to finish.
 			{
-				FLinearColor& Pixel = OutputFrame->OutputEquirectangularMap[PixelIndex];
-				Pixel.R /= Pixel.A;
-				Pixel.G /= Pixel.A;
-				Pixel.B /= Pixel.A;
-				
-				// Now that we've used the weight stored in the alpha channel, restore it to full opaque-ness so
-				// the resulting png/exrs don't end up semi-transparent.
-				Pixel.A = 1.f;
+				SCOPE_CYCLE_COUNTER(STAT_MoviePipeline_PanoBlendWait);
+				TargetBlender->Blender.TaskConcurrencyLimiter->Wait();
 			}
-		}
-		TSharedRef<FImagePixelDataPayload, ESPMode::ThreadSafe> NewPayload = DataPayload->Copy();
+			if (ensure(OutputMerger.IsValid()))
+			{
+				TUniquePtr<TImagePixelData<FLinearColor> > FinalPixelData = MakeUnique<TImagePixelData<FLinearColor>>(OutputResolution, DataPayload->Copy());
+				TargetBlender->Blender.FetchFinalPixelDataLinearColor(FinalPixelData->Pixels);
 
-		int32 OutputSizeX = OutputEquirectangularMapSize.X;
-		int32 OutputSizeY = DataPayload->Pane.Data.EyeIndex >= 0 ? OutputEquirectangularMapSize.Y * 2 : OutputEquirectangularMapSize.Y;
-		TUniquePtr<TImagePixelData<FLinearColor>> FinalPixelData = MakeUnique<TImagePixelData<FLinearColor>>(FIntPoint(OutputSizeX, OutputSizeY), TArray64<FLinearColor>(MoveTemp(OutputFrame->OutputEquirectangularMap)), NewPayload);
+				OutputMerger.Pin()->OnCompleteRenderPassDataAvailable_AnyThread(MoveTemp(FinalPixelData));
+			}
 
-		if(ensure(OutputMerger.IsValid()))
-		{
-			OutputMerger.Pin()->OnCompleteRenderPassDataAvailable_AnyThread(MoveTemp(FinalPixelData));
+			// Release the pool item so future frames can use it.
+			TargetBlender->bActive = false;
 		}
-
-		// Remove this frame from the PendingData map so that we don't hold references for the duration of the render.
-		// It's important that we use the manually looked up OutputFrame from PendingData
-		// as PendingData uses the equality operator. Some combinations of temporal sampling + slowmo tracks results in different
-		// original source frame numbers, which would cause the tmap lookup to fail and thus returning an empty frame.
-		PendingData.Remove(DataPayload->SampleState.OutputState);
 	}
-
 }
 
 void FMoviePipelinePanoramicBlender::OnSingleSampleDataAvailable_AnyThread(TUniquePtr<FImagePixelData>&& InData)
@@ -431,3 +162,10 @@ void FMoviePipelinePanoramicBlender::AbandonOutstandingWork()
 	// Not yet implemented
 	check(0);
 }
+
+int32 FMoviePipelinePanoramicBlender::GetNumOutstandingFrames() const
+{
+	return 0;
+}
+
+} // UE::MoviePipeline
diff --git a/Engine/Plugins/MovieScene/MoviePipelineMaskRenderPass/Source/MoviePipelineMaskRenderPass/Private/MoviePipelinePanoramicBlender.h b/Engine/Plugins/MovieScene/MoviePipelineMaskRenderPass/Source/MoviePipelineMaskRenderPass/Private/MoviePipelinePanoramicBlender.h
--- a/Engine/Plugins/MovieScene/MoviePipelineMaskRenderPass/Source/MoviePipelineMaskRenderPass/Private/MoviePipelinePanoramicBlender.h
+++ b/Engine/Plugins/MovieScene/MoviePipelineMaskRenderPass/Source/MoviePipelineMaskRenderPass/Private/MoviePipelinePanoramicBlender.h
@@ -2,55 +2,47 @@
 #pragma once
 
 #include "MovieRenderPipelineDataTypes.h"
+#include "MoviePipelinePanoramicBlenderBase.h"
 
 // Forward Declares
 struct FImagePixelData;
 class UMoviePipeline;
 
-class FMoviePipelinePanoramicBlender : public MoviePipeline::IMoviePipelineOutputMerger
+namespace UE::MoviePipeline
 {
-public:
-	FMoviePipelinePanoramicBlender(TSharedPtr<MoviePipeline::IMoviePipelineOutputMerger> InOutputMerger, const FIntPoint InOutputResolution);
-public:
-	virtual FMoviePipelineMergerOutputFrame& QueueOutputFrame_GameThread(const FMoviePipelineFrameOutputState& CachedOutputState) override;
-	virtual void OnCompleteRenderPassDataAvailable_AnyThread(TUniquePtr<FImagePixelData>&& InData) override;
-	virtual void OnSingleSampleDataAvailable_AnyThread(TUniquePtr<FImagePixelData>&& InData) override;
-	virtual void AbandonOutstandingWork() override;
-	virtual int32 GetNumOutstandingFrames() const override { return PendingData.Num(); }
-
-private:
-	struct FPanoramicBlendData
+	/**
+	* The FMoviePipelinePanoramicBlender acts as an intermediate destination before the regular OutputMerger of the MoviePipeline. To reuse the multi-sample accumulation code,
+	* we run those accumulations and tell them to send the results to the output merger, but we've replaced the standard Movie Pipeline output merger with an instance of
+	* this class, and it then takes all the samples for the frame and blends them into one larger output image. Then once all of the samples come in for a given output image,
+	* it is passed onto the regular MoviePipeline output merger.
+	*/
+	class FMoviePipelinePanoramicBlender : public ::MoviePipeline::IMoviePipelineOutputMerger
 	{
-		double BlendStartTime;
-		double BlendEndTime;
-		bool bFinished;
-		FIntPoint OutputBoundsMin;
-		FIntPoint OutputBoundsMax;
-		int32 PixelWidth;
-		int32 PixelHeight;
-		TArray<FLinearColor> Data;
-		int32 EyeIndex;
-		TSharedPtr<struct FPanoramicImagePixelDataPayload> OriginalDataPayload;
+	public:
+		FMoviePipelinePanoramicBlender(TSharedPtr<::MoviePipeline::IMoviePipelineOutputMerger> InOutputMerger, const FIntPoint InOutputResolution);
+	public:
+		virtual FMoviePipelineMergerOutputFrame& QueueOutputFrame_GameThread(const FMoviePipelineFrameOutputState& CachedOutputState) override;
+		virtual void OnCompleteRenderPassDataAvailable_AnyThread(TUniquePtr<FImagePixelData>&& InData) override;
+		virtual void OnSingleSampleDataAvailable_AnyThread(TUniquePtr<FImagePixelData>&& InData) override;
+		virtual void AbandonOutstandingWork() override;
+		virtual int32 GetNumOutstandingFrames() const override;
+	private:
+		struct FPoolEntry
+		{
+			FMoviePipelinePanoramicBlenderBase Blender;
+			bool bActive;
+			int32 OutputFrameNumber;
+			std::atomic<int32> NumCompletedAccumulations;
+		};
+
+		// Pool entries are allocated as pointers on the heap so that if the array is resized while a thread is
+		// working on a previous frame, it doesn't have the memory moved out from underneath it.
+		TArray<TUniquePtr<FPoolEntry>> PendingData;
+
+		FCriticalSection GlobalQueueDataMutex;
+
+		FIntPoint OutputResolution;
+
+		TWeakPtr<::MoviePipeline::IMoviePipelineOutputMerger> OutputMerger;
 	};
-
-	struct FPanoramicOutputFrame
-	{
-		// Eye Index to Blend Data. Eye Index will be -1 when not using Stereo.
-		TMap<int32, TArray<TSharedPtr<FPanoramicBlendData>>> BlendedData;
-
-		// The total number of samples we have to wait for to finish blending before being 'done'.
-		int32 NumSamplesTotal;
-
-		TArray<FLinearColor> OutputEquirectangularMap;
-	};
-
-	/** Data that is expected but not fully available yet. */
-	TMap<FMoviePipelineFrameOutputState, TSharedPtr<FPanoramicOutputFrame>> PendingData;
-
-	/** Mutex that protects adding/updating/removing from PendingData */
-	FCriticalSection GlobalQueueDataMutex;
-	FCriticalSection OutputDataMutex;
-	FIntPoint OutputEquirectangularMapSize;
-
-	TWeakPtr<MoviePipeline::IMoviePipelineOutputMerger> OutputMerger;
-};
+}
diff --git a/Engine/Plugins/MovieScene/MoviePipelineMaskRenderPass/Source/MoviePipelineMaskRenderPass/Private/MoviePipelinePanoramicPass.cpp b/Engine/Plugins/MovieScene/MoviePipelineMaskRenderPass/Source/MoviePipelineMaskRenderPass/Private/MoviePipelinePanoramicPass.cpp
--- a/Engine/Plugins/MovieScene/MoviePipelineMaskRenderPass/Source/MoviePipelineMaskRenderPass/Private/MoviePipelinePanoramicPass.cpp
+++ b/Engine/Plugins/MovieScene/MoviePipelineMaskRenderPass/Source/MoviePipelineMaskRenderPass/Private/MoviePipelinePanoramicPass.cpp
@@ -45,7 +45,7 @@ void UMoviePipelinePanoramicPass::MoviePipelineRenderShowFlagOverride(FEngineSho
 	OutShowFlag.SetVignette(false);
 	OutShowFlag.SetSceneColorFringe(false);
 	OutShowFlag.SetPhysicalMaterialMasks(false);
-
+	OutShowFlag.SetDepthOfField(false);
 	/*if(bPathTracer)
 	{
 		OutShowFlag.SetPathTracing(true);
@@ -119,7 +119,7 @@ void UMoviePipelinePanoramicPass::SetupImpl(const MoviePipeline::FMoviePipelineR
 	// Create a class to blend the Panoramic Panes into equirectangular maps. When all of the samples for a given
 	// Panorama are provided to the Blender and it is blended, it will pass the data onto the normal OutputBuilder
 	// who is none the wiser that we're handing it complex blended images instead of normal stills.
-	PanoramicOutputBlender = MakeShared<FMoviePipelinePanoramicBlender>(GetPipeline()->OutputBuilder, InPassInitSettings.BackbufferResolution);
+	PanoramicOutputBlender = MakeShared<UE::MoviePipeline::FMoviePipelinePanoramicBlender>(GetPipeline()->OutputBuilder, InPassInitSettings.BackbufferResolution);
 
 	// Allocate an OCIO extension to do color grading if needed.
 	OCIOSceneViewExtension = FSceneViewExtensions::NewExtension<FOpenColorIODisplayExtension>();
@@ -476,12 +476,15 @@ void UMoviePipelinePanoramicPass::RenderSample_GameThreadImpl(const FMoviePipeli
 	// Wait for a surface to be available to write to. This will stall the game thread while the RHI/Render Thread catch up.
 	Super::RenderSample_GameThreadImpl(InSampleState);
 
+	UMoviePipelinePrimaryConfig* PrimaryConfig = GetPipeline()->GetPipelinePrimaryConfig();
+	UMoviePipelineExecutorShot* CurrentShot = GetPipeline()->GetActiveShotList()[GetPipeline()->GetCurrentShotIndex()];
+
 	const FIntPoint PaneResolution = GetPaneResolution(InSampleState.BackbufferSize);
 	
 	// Auto-exposure is rendered for the first sample, and shared across all panes
 	if (AutoExposureViewStates.Num() && InSampleState.SpatialSampleIndex == 0 && InSampleState.TemporalSampleIndex == 0)
 	{
-		FMoviePipelineRenderPassMetrics InOutSampleState = InSampleState;
+		FMoviePipelineRenderPassMetrics InOutSampleState = UE::MoviePipeline::GetRenderPassMetrics(PrimaryConfig, CurrentShot, InSampleState, PaneResolution);
 		InOutSampleState.bAutoExposureCubePass = true;
 		InOutSampleState.BackbufferSize = PaneResolution;
 		InOutSampleState.TileIndexes = FIntPoint(-1, -1);		// Marks this as the auto exposure pass
@@ -500,6 +503,7 @@ void UMoviePipelinePanoramicPass::RenderSample_GameThreadImpl(const FMoviePipeli
 		Pane.Data.EyeSeparation = EyeSeparation;
 		Pane.Data.EyeConvergenceDistance = EyeConvergenceDistance;
 		Pane.Data.bUseLocalRotation = bFollowCameraOrientation;
+		Pane.Data.FilterType = Filter;
 		
 		UE::MoviePipeline::Panoramic::GetCameraOrientationForStereo(/*Out*/ Pane.Data.CameraLocation, /*Out*/ Pane.Data.CameraRotation, /*Out*/ Pane.Data.CameraLocalRotation, Pane.Data, -1, /*bInPrevPos*/ false);
 		FRotator DummyPrevLocalRot;
@@ -538,7 +542,8 @@ void UMoviePipelinePanoramicPass::RenderSample_GameThreadImpl(const FMoviePipeli
 			int32 NumEyeRenders = bStereo ? 2 : 1;
 			for (int32 EyeLoopIndex = 0; EyeLoopIndex < NumEyeRenders; EyeLoopIndex++)
 			{
-				FMoviePipelineRenderPassMetrics InOutSampleState = InSampleState;
+
+				FMoviePipelineRenderPassMetrics InOutSampleState = UE::MoviePipeline::GetRenderPassMetrics(PrimaryConfig, CurrentShot, InSampleState, PaneResolution);
 
 				// Construct a pane that will be attached to the render information
 				FPanoPane Pane;
@@ -556,6 +561,7 @@ void UMoviePipelinePanoramicPass::RenderSample_GameThreadImpl(const FMoviePipeli
 				Pane.Data.EyeSeparation = EyeSeparation;
 				Pane.Data.EyeConvergenceDistance = EyeConvergenceDistance;
 				Pane.Data.bUseLocalRotation = bFollowCameraOrientation;
+				Pane.Data.FilterType = Filter;
 
 				// Get the actual camera location/rotation for this particular pane, the above values are from the global camera.
 				UE::MoviePipeline::Panoramic::GetCameraOrientationForStereo(/*Out*/ Pane.Data.CameraLocation, /*Out*/ Pane.Data.CameraRotation, /*Out*/ Pane.Data.CameraLocalRotation, Pane.Data, StereoIndex, /*bInPrevPos*/ false);
@@ -576,26 +582,26 @@ void UMoviePipelinePanoramicPass::RenderSample_GameThreadImpl(const FMoviePipeli
 					SceneViewState->SystemMemoryMirrorRestore(SystemMemoryMirror.Get());
 				}
 
-				// Create the View Family for this render. This will only contain one view to better fit into the existing
-				// architecture of MRQ we have. Calculating the view family requires calculating the FSceneView itself which
-				// is highly customized for panos. So we provide the FPanoPlane as 'raw' data that gets passed along so we can
-				// use it when we calculate our individual view.
-				TSharedPtr<FSceneViewFamilyContext> ViewFamily = CalculateViewFamily(InOutSampleState, &Pane);
-
-				EAntiAliasingMethod AAMethod = ViewFamily->Views[0]->AntiAliasingMethod;
-				const bool bRequiresHistory = (AAMethod == EAntiAliasingMethod::AAM_TemporalAA) || (AAMethod == EAntiAliasingMethod::AAM_TSR);
+				// We need to check and set the AA method first before calculating the view family, otherwise it uses the wrong
+				// AA method and doesn't apply jitter when it should.
+				const bool bRequiresHistory = (InOutSampleState.AntiAliasingMethod == EAntiAliasingMethod::AAM_TemporalAA) || (InOutSampleState .AntiAliasingMethod == EAntiAliasingMethod::AAM_TSR);
 				if (!bAllocateHistoryPerPane && bRequiresHistory)
 				{
 					if (!bHasWarnedSettings)
 					{
 						bHasWarnedSettings = true;
-						UE_LOG(LogMovieRenderPipeline, Warning, TEXT("Panoramic Renders do not support TAA without enabling bAllocateHistoryPerPane! Forcing AntiAliasing off."));
+						UE_LOG(LogMovieRenderPipeline, Warning, TEXT("Panoramic Renders do not support TAA/TSR without enabling bAllocateHistoryPerPane! Forcing AntiAliasing off."));
 					}
 					
-					FSceneView* NonConstView = const_cast<FSceneView*>(ViewFamily->Views[0]);
-					NonConstView->AntiAliasingMethod = EAntiAliasingMethod::AAM_None;
+					InOutSampleState.AntiAliasingMethod = EAntiAliasingMethod::AAM_None;
 				}
 
+				// Create the View Family for this render. This will only contain one view to better fit into the existing
+				// architecture of MRQ we have. Calculating the view family requires calculating the FSceneView itself which
+				// is highly customized for panos. So we provide the FPanoPlane as 'raw' data that gets passed along so we can
+				// use it when we calculate our individual view.
+				TSharedPtr<FSceneViewFamilyContext> ViewFamily = CalculateViewFamily(InOutSampleState, &Pane);
+
 				// Submit the view to be rendered.
 				TWeakObjectPtr<UTextureRenderTarget2D> ViewRenderTarget = GetOrCreateViewRenderTarget(PaneResolution);
 				check(ViewRenderTarget.IsValid());
@@ -644,7 +650,7 @@ void UMoviePipelinePanoramicPass::ScheduleReadbackAndAccumulation(const FMoviePi
 	{
 		SCOPE_CYCLE_COUNTER(STAT_MoviePipeline_WaitForAvailableAccumulator);
 		
-		// Generate a unique PassIdentifier for this Panoramic Pane.Data. High Res tile isn't taken into account because
+		// Generate a unique PassIdentifier for this Panoramic Pane. High Res tile isn't taken into account because
 		// the same accumulator is used for all tiles.
 		FMoviePipelinePassIdentifier PanePassIdentifier = FMoviePipelinePassIdentifier(FString::Printf(TEXT("%s_x%d_y%d"), *PassIdentifier.Name, InPane.Data.VerticalStepIndex, InPane.Data.HorizontalStepIndex));
 		SampleAccumulator = AccumulatorPool->BlockAndGetAccumulator_GameThread(InSampleState.OutputState.OutputFrameNumber, PanePassIdentifier);
diff --git a/Engine/Plugins/MovieScene/MoviePipelineMaskRenderPass/Source/MoviePipelineMaskRenderPass/Private/MoviePipelinePanoramicPass.h b/Engine/Plugins/MovieScene/MoviePipelineMaskRenderPass/Source/MoviePipelineMaskRenderPass/Private/MoviePipelinePanoramicPass.h
--- a/Engine/Plugins/MovieScene/MoviePipelineMaskRenderPass/Source/MoviePipelineMaskRenderPass/Private/MoviePipelinePanoramicPass.h
+++ b/Engine/Plugins/MovieScene/MoviePipelineMaskRenderPass/Source/MoviePipelineMaskRenderPass/Private/MoviePipelinePanoramicPass.h
@@ -2,6 +2,7 @@
 #pragma once
 
 #include "MoviePipelineImagePassBase.h"
+#include "MoviePipelinePanoramicBlenderBase.h"
 #include "Templates/PimplPtr.h"
 #include "MoviePipelinePanoramicPass.generated.h"
 
@@ -121,6 +122,12 @@ public:
 	UPROPERTY(EditAnywhere, BlueprintReadWrite, Category = "Panoramic Settings")
 	bool bFollowCameraOrientation;
 	
+	/**
+	* Filter used when blending panoramic. Bilinear is fastest (Samples a 2x2 pixel grid) and produces nearly as good results as the others which require sampling 4x4 pixels.
+	*/
+	UPROPERTY(EditAnywhere, BlueprintReadWrite, Category = "Panoramic Settings")
+	EMoviePipelinePanoramicFilterType Filter = EMoviePipelinePanoramicFilterType::Bilinear;
+	
 	/**
 	* If true we will capture a stereo panorama. This doubles everything (all render times, memory requirements, etc.)
 	* May need a significantly higher number of horizontal steps to look good when viewed in stereo.
diff --git a/Engine/Plugins/MovieScene/MovieRenderPipeline/Source/MovieRenderPipelineCore/Private/Graph/MovieGraphOutputMerger.cpp b/Engine/Plugins/MovieScene/MovieRenderPipeline/Source/MovieRenderPipelineCore/Private/Graph/MovieGraphOutputMerger.cpp
--- a/Engine/Plugins/MovieScene/MovieRenderPipeline/Source/MovieRenderPipelineCore/Private/Graph/MovieGraphOutputMerger.cpp
+++ b/Engine/Plugins/MovieScene/MovieRenderPipeline/Source/MovieRenderPipelineCore/Private/Graph/MovieGraphOutputMerger.cpp
@@ -64,15 +64,25 @@ namespace UE::MovieGraph
 			FString FinalFilePath;
 			{
 				// TODO: Add Tile X/Y when tiling is available
-				const FString OutputName = FString::Printf(TEXT("%s_%s_%s_%s_%s_SS_%d_TS_%d.%d"),
-					*TraversalContext.Shot->OuterName,
-					*TraversalContext.RenderDataIdentifier.LayerName,
-					*TraversalContext.RenderDataIdentifier.RendererName,
-					*TraversalContext.RenderDataIdentifier.SubResourceName,
-					*TraversalContext.RenderDataIdentifier.CameraName,
-					TraversalContext.Time.SpatialSampleIndex,
-					TraversalContext.Time.TemporalSampleIndex,
-					TraversalContext.Time.OutputFrameNumber);
+				FString OutputName;
+				if (Payload->Debug_OverrideFilename.IsEmpty())
+				{
+
+					OutputName = FString::Printf(TEXT("%s_%s_%s_%s_%s_SS_%d_TS_%d.%d"),
+						*TraversalContext.Shot->OuterName,
+						*TraversalContext.RenderDataIdentifier.LayerName,
+						*TraversalContext.RenderDataIdentifier.RendererName,
+						*TraversalContext.RenderDataIdentifier.SubResourceName,
+						*TraversalContext.RenderDataIdentifier.CameraName,
+						TraversalContext.Time.SpatialSampleIndex,
+						TraversalContext.Time.TemporalSampleIndex,
+						TraversalContext.Time.OutputFrameNumber);
+				}
+				else
+				{
+					// If provided, we assume they've already resolved the formatting for the exact name they want.
+					OutputName = Payload->Debug_OverrideFilename;
+				}
 
 				const UMovieGraphGlobalOutputSettingNode* OutputSettings =
 					EvaluatedConfig->GetSettingForBranch<UMovieGraphGlobalOutputSettingNode>(UMovieGraphNode::GlobalsPinName);
diff --git a/Engine/Plugins/MovieScene/MovieRenderPipeline/Source/MovieRenderPipelineCore/Private/MoviePipelineRendering.cpp b/Engine/Plugins/MovieScene/MovieRenderPipeline/Source/MovieRenderPipelineCore/Private/MoviePipelineRendering.cpp
--- a/Engine/Plugins/MovieScene/MovieRenderPipeline/Source/MovieRenderPipelineCore/Private/MoviePipelineRendering.cpp
+++ b/Engine/Plugins/MovieScene/MovieRenderPipeline/Source/MovieRenderPipelineCore/Private/MoviePipelineRendering.cpp
@@ -53,6 +53,17 @@ static TAutoConsoleVariable<int> CVarMoviePipelineThrottleFrameCount(
 	TEXT("This option only applies to path traced renders, as deferred rendering is synchronized through pixel readbacks.\n"),
 	ECVF_Default);
 
+static TAutoConsoleVariable<bool> CVarMoviePipelineDisableMaxResolutionCheck(
+	TEXT("MoviePipeline.DisableMaxResolutionCheck"),
+	false,
+	TEXT("When true, the Movie Pipeline will not consider an output resolution higher than the maximum texture size an error. "
+		"When using high resolution tiles the system can calculate the size of an individual tile to correctly detect an overly "
+		"large render, but when using panoramic renders the system cannot, artificially limiting the resolution to 16k. Set this "
+		"true to disable the check, with the understanding that user settings can then choose crash-inducing values."),
+	ECVF_Default
+);
+
+
 #define LOCTEXT_NAMESPACE "MoviePipeline"
 
 void UMoviePipeline::SetupRenderingPipelineForShot(UMoviePipelineExecutorShot* InShot)
@@ -101,11 +112,17 @@ void UMoviePipeline::SetupRenderingPipelineForShot(UMoviePipelineExecutorShot* I
 		int32 MaxResolution = GetMax2DTextureDimension();
 		if (BackbufferResolution.X > MaxResolution || BackbufferResolution.Y > MaxResolution)
 		{
-			UE_LOG(LogMovieRenderPipeline, Error, TEXT("Resolution %dx%d exceeds maximum allowed by GPU (%dx%d). Consider using the HighRes setting and increasing the tile count."), BackbufferResolution.X, BackbufferResolution.Y, MaxResolution, MaxResolution);
-			Shutdown(true);
-			return;
+			// Panoramic Tiling doesn't correctly pass this check so the system unnecessarily prevents users
+			// from making large panoramic renders. This cvar allows the user to disable this check, but we
+			// keep the log above as it contains good information to have.
+			if (!CVarMoviePipelineDisableMaxResolutionCheck.GetValueOnGameThread())
+			{
+				UE_LOG(LogMovieRenderPipeline, Error, TEXT("Resolution %dx%d exceeds maximum allowed by GPU (%dx%d). Consider using the HighRes setting and increasing the tile count. If using a tiled render pass that does not use HighRes tiles (such as the Panoramic Pass), set MoviePipeline.DisableMaxResolutionCheck to true to bypass this and allow rendering."), BackbufferResolution.X, BackbufferResolution.Y, MaxResolution, MaxResolution);
+				Shutdown(true);
+				return;
+			}
 		}
-		if (BackbufferResolution.X == 0 || BackbufferResolution.Y == 0)
+		if (BackbufferResolution.X <= 0 || BackbufferResolution.Y <= 0)
 		{
 			UE_LOG(LogMovieRenderPipeline, Error, TEXT("Resolution %dx%d must be greater than zero in both dimensions."), BackbufferResolution.X, BackbufferResolution.Y);
 			Shutdown(true);
diff --git a/Engine/Plugins/MovieScene/MovieRenderPipeline/Source/MovieRenderPipelineCore/Public/Graph/MovieGraphDataTypes.h b/Engine/Plugins/MovieScene/MovieRenderPipeline/Source/MovieRenderPipelineCore/Public/Graph/MovieGraphDataTypes.h
--- a/Engine/Plugins/MovieScene/MovieRenderPipeline/Source/MovieRenderPipelineCore/Public/Graph/MovieGraphDataTypes.h
+++ b/Engine/Plugins/MovieScene/MovieRenderPipeline/Source/MovieRenderPipelineCore/Public/Graph/MovieGraphDataTypes.h
@@ -432,6 +432,12 @@ namespace UE::MovieGraph
 		/** Debug feature, should this sample be written to disk, not being accumulated? */
 		bool bWriteSampleToDisk;
 
+		/**
+		* When bWriteSampleToDisk is true, if this string is specified then the single-sample debug writing will use this filename for the format
+		* token instead of the default one. Useful for render passes with complex output cases (like Panoramics) that need pass-specific names.
+		*/
+		FString Debug_OverrideFilename;
+
 		/** Set this to true for every sample if an accumulator is required (ie: has temporal/spatial sub-samples, or high res tiling). */
 		bool bRequiresAccumulator;
 
diff --git a/Engine/Plugins/MovieScene/MovieRenderPipeline/Source/MovieRenderPipelineCore/Public/MoviePipelinePanoramicBlenderBase.cpp b/Engine/Plugins/MovieScene/MovieRenderPipeline/Source/MovieRenderPipelineCore/Public/MoviePipelinePanoramicBlenderBase.cpp
new file mode 100644
--- /dev/null
+++ b/Engine/Plugins/MovieScene/MovieRenderPipeline/Source/MovieRenderPipelineCore/Public/MoviePipelinePanoramicBlenderBase.cpp
@@ -0,0 +1,527 @@
+// Copyright Epic Games, Inc. All Rights Reserved.
+#include "MoviePipelinePanoramicBlenderBase.h"
+#include "Math/PerspectiveMatrix.h"
+#include "HAL/UnrealMemory.h"
+#include "MovieRenderPipelineCoreModule.h"
+
+static TAutoConsoleVariable<int32> CVarMoviePipelinePanoramicMaxPoolsPerFrame(
+	TEXT("MoviePipeline.Panoramic.MaxConcurrentBlendingPoolCount"),
+	4,
+	TEXT("When blending panoramic images, this determines the maximum number of concurrent blending pools. "
+		"A larger number may result in better CPU usage, but can come at a significant cost to CPU memory. "
+		"Lowering this value can reduce the amount of memory needed for panoramic blending, but will result in "
+		"blending taking longer."),
+	ECVF_Default
+);
+
+namespace UE::MoviePipeline
+{
+	static FLinearColor GetColorBilinearFiltered(const FImagePixelData* InSampleData, const FVector2D& InSamplePixelCoords, bool& OutClipped, bool bInForceAlphaToOpaque = false)
+	{
+		// Pixel coordinates assume that 0.5, 0.5 is the center of the pixel, so we subtract half to make it indexable.
+		const FVector2D PixelCoordinateIndex = InSamplePixelCoords - 0.5f;
+
+		// Get surrounding pixels indices
+		FIntPoint LowerLeftPixelIndex = FIntPoint(FMath::FloorToInt(PixelCoordinateIndex.X), FMath::FloorToInt(PixelCoordinateIndex.Y));
+		FIntPoint LowerRightPixelIndex = FIntPoint(LowerLeftPixelIndex + FIntPoint(1, 0));
+		FIntPoint UpperLeftPixelIndex = FIntPoint(LowerLeftPixelIndex + FIntPoint(0, 1));
+		FIntPoint UpperRightPixelIndex = FIntPoint(LowerLeftPixelIndex + FIntPoint(1, 1));
+
+		// Clamp pixels indices to pixels array bounds.
+		// ToDo: Is this needed? Should we handle wrap around for the bottom right pixel? What gives
+		auto ClampPixelCoords = [&](FIntPoint& InOutPixelCoords, const FIntPoint& InArraySize)
+			{
+				if (InOutPixelCoords.X > InArraySize.X - 1 ||
+					InOutPixelCoords.Y > InArraySize.Y - 1 ||
+					InOutPixelCoords.X < 0 ||
+					InOutPixelCoords.Y < 0)
+				{
+					OutClipped = true;
+				}
+				InOutPixelCoords = FIntPoint(FMath::Clamp(InOutPixelCoords.X, 0, InArraySize.X - 1), FMath::Clamp(InOutPixelCoords.Y, 0, InArraySize.Y - 1));
+			};
+		ClampPixelCoords(LowerLeftPixelIndex, InSampleData->GetSize());
+		ClampPixelCoords(LowerRightPixelIndex, InSampleData->GetSize());
+		ClampPixelCoords(UpperLeftPixelIndex, InSampleData->GetSize());
+		ClampPixelCoords(UpperRightPixelIndex, InSampleData->GetSize());
+
+		// Fetch the colors for the four pixels. We convert to FLinearColor here so that our accumulation
+		// is done in linear space with enough precision. The samples are probably in F16 color right now.
+		FLinearColor LowerLeftPixelColor;
+		FLinearColor LowerRightPixelColor;
+		FLinearColor UpperLeftPixelColor;
+		FLinearColor UpperRightPixelColor;
+
+		int64 SizeInBytes = 0;
+		const void* SrcRawDataPtr = nullptr;
+		InSampleData->GetRawData(SrcRawDataPtr, SizeInBytes);
+
+		switch (InSampleData->GetType())
+		{
+		case EImagePixelType::Float16:
+		{
+			const FFloat16Color* ColorDataF16 = static_cast<const FFloat16Color*>(SrcRawDataPtr);
+			LowerLeftPixelColor = FLinearColor(ColorDataF16[LowerLeftPixelIndex.X + (LowerLeftPixelIndex.Y * InSampleData->GetSize().X)]);
+			LowerRightPixelColor = FLinearColor(ColorDataF16[LowerRightPixelIndex.X + (LowerRightPixelIndex.Y * InSampleData->GetSize().X)]);
+			UpperLeftPixelColor = FLinearColor(ColorDataF16[UpperLeftPixelIndex.X + (UpperLeftPixelIndex.Y * InSampleData->GetSize().X)]);
+			UpperRightPixelColor = FLinearColor(ColorDataF16[UpperRightPixelIndex.X + (UpperRightPixelIndex.Y * InSampleData->GetSize().X)]);
+		}
+		break;
+		case EImagePixelType::Float32:
+		{
+			const FLinearColor* ColorDataF32 = static_cast<const FLinearColor*>(SrcRawDataPtr);
+			LowerLeftPixelColor = ColorDataF32[LowerLeftPixelIndex.X + (LowerLeftPixelIndex.Y * InSampleData->GetSize().X)];
+			LowerRightPixelColor = ColorDataF32[LowerRightPixelIndex.X + (LowerRightPixelIndex.Y * InSampleData->GetSize().X)];
+			UpperLeftPixelColor = ColorDataF32[UpperLeftPixelIndex.X + (UpperLeftPixelIndex.Y * InSampleData->GetSize().X)];
+			UpperRightPixelColor = ColorDataF32[UpperRightPixelIndex.X + (UpperRightPixelIndex.Y * InSampleData->GetSize().X)];
+		}
+		break;
+		default:
+			// Not implemented
+			check(0);
+		}
+
+		const float FracX = (InSamplePixelCoords.X - LowerLeftPixelIndex.X - 0.5f);
+		const float FracY = (InSamplePixelCoords.Y - LowerLeftPixelIndex.Y - 0.5f);
+
+		FLinearColor InterpolatedPixelColor =
+			(LowerLeftPixelColor * (1.0f - FracX) + LowerRightPixelColor * FracX) * (1.0f - FracY)
+			+ (UpperLeftPixelColor * (1.0f - FracX) + UpperRightPixelColor * FracX) * FracY;
+
+		// Force final color alpha to opaque if requested
+		if (bInForceAlphaToOpaque)
+		{
+			InterpolatedPixelColor.A = 1.0f;
+		}
+
+		return InterpolatedPixelColor;
+	}
+
+
+static FLinearColor GetColorCubicFiltered(const FImagePixelData* InSampleData, const FVector2D& InSamplePixelCoords, float CubicBParam, float CubicCParam, bool& OutClipped, bool bInForceAlphaToOpaque = false)
+{
+	// Pixel coordinates assume that 0.5, 0.5 is the center of the pixel, so we subtract half to make it indexable.
+	const FVector2D PixelCoordinateIndex = InSamplePixelCoords - 0.5f;
+
+	// Get surrounding 4x4 pixels indices, because we floored our center is off-center to the lower-left.
+	FIntPoint PixelCenter = FIntPoint(FMath::FloorToInt(PixelCoordinateIndex.X), FMath::FloorToInt(PixelCoordinateIndex.Y));
+
+	auto ParameterizedCubic = [](float X, float B, float C)
+		{
+			float AbsX = FMath::Abs(X);
+			float X2 = AbsX * AbsX;
+			float X3 = AbsX * AbsX * AbsX;
+
+			if (AbsX <= 1.f)
+			{
+				return((12 - 9 * B - 6 * C) * X3 +
+					(-18 + 12 * B + 6 * C) * X2 +
+					(6 - 2 * B)) / 6.0f;
+			}
+			else if (AbsX < 2.0f)
+			{
+				return ((-B - 6 * C) * X3 +
+					(6 * B + 30 * C) * X2 +
+					(-12 * B - 48 * C) * AbsX +
+					(8 * B + 24 * C)) / 6.0f;
+			}
+			else
+			{
+				return 0.f;
+			}
+		};
+	// The fractional amount we were within the pixel
+	float FracX = InSamplePixelCoords.X - PixelCenter.X - 0.5f;
+	float FracY = InSamplePixelCoords.Y - PixelCenter.Y - 0.5f;
+
+	float WeightX[4];
+	float WeightY[4];
+	for (int32 i = 0; i < 4; i++)
+	{
+		WeightX[i] = ParameterizedCubic((i - 1) - FracX, CubicBParam, CubicCParam);
+		WeightY[i] = ParameterizedCubic((i - 1) - FracY, CubicBParam, CubicCParam);
+	}
+
+	int64 SizeInBytes = 0;
+	const void* SrcRawDataPtr = nullptr;
+	InSampleData->GetRawData(SrcRawDataPtr, SizeInBytes);
+
+	OutClipped = false;
+
+	FLinearColor ResultingColor = FLinearColor(0, 0, 0, 0);
+	for (int32 i = 0; i < 4; i++)
+	{
+		int32 yIndex = FMath::Clamp(PixelCenter.Y + i - 1, 0, InSampleData->GetSize().Y - 1);
+		for (int32 j = 0; j < 4; j++)
+		{
+			int32 xIndex = FMath::Clamp(PixelCenter.X + j - 1, 0, InSampleData->GetSize().X - 1);
+			int32 PixelIndex = (yIndex * InSampleData->GetSize().X) + xIndex;
+			
+			FLinearColor Sample = FLinearColor(0, 0, 0, 0);
+			switch (InSampleData->GetType())
+			{
+			case EImagePixelType::Float16:
+			{
+				const FFloat16Color* ColorDataF16 = static_cast<const FFloat16Color*>(SrcRawDataPtr);
+				Sample = FLinearColor(ColorDataF16[PixelIndex]);
+				break;
+			}
+			case EImagePixelType::Float32:
+			{
+				const FLinearColor* ColorDataF32 = static_cast<const FLinearColor*>(SrcRawDataPtr);
+				Sample = ColorDataF32[PixelIndex];
+				break;
+			}
+			}
+
+			ResultingColor += WeightX[j] * WeightY[i] * Sample;
+		}
+	}
+
+	// Force final color alpha to opaque if requested
+	if (bInForceAlphaToOpaque)
+	{
+		ResultingColor.A = 1.0f;
+	}
+
+	return ResultingColor;
+}
+
+void FMoviePipelinePanoramicBlenderBase::Initialize(const FIntPoint InOutputResolution)
+{
+	OutputEquirectangularMapSize = InOutputResolution;
+
+	// If these are being re-used this should be a no-op
+	OutputEquirectangularMap.SetNumUninitialized(OutputEquirectangularMapSize.X * OutputEquirectangularMapSize.Y);
+
+	// We need to zero the memory though as the results are accumulated into it.
+	FMemory::Memzero(OutputEquirectangularMap.GetData(), OutputEquirectangularMap.Num() * OutputEquirectangularMap.GetTypeSize());
+
+	// Tasks get pushed into a concurrency limiter to avoid allocating too many concurrent blend pools which can take significant RAM on large images.
+	// There is a fair amount of parallelism within each task, so high core count machines still get high occupancy even with low concurrency.
+	TaskConcurrencyLimiter = MakeUnique<UE::Tasks::FTaskConcurrencyLimiter>(FMath::Max(CVarMoviePipelinePanoramicMaxPoolsPerFrame.GetValueOnAnyThread(), 1));
+}
+
+DECLARE_CYCLE_STAT(TEXT("STAT_MoviePipeline_PanoBlend"), STAT_MoviePipeline_PanoBlend, STATGROUP_MoviePipeline);
+DECLARE_CYCLE_STAT(TEXT("STAT_MoviePipeline_PanoBlendActual"), STAT_MoviePipeline_PanoBlendActual, STATGROUP_MoviePipeline);
+DECLARE_CYCLE_STAT(TEXT("STAT_MoviePipeline_PanoBlendJoin"), STAT_MoviePipeline_PanoBlendJoin, STATGROUP_MoviePipeline);
+
+void FMoviePipelinePanoramicBlenderBase::BlendSample_AnyThread(TUniquePtr<FImagePixelData>&& InData, const UE::MoviePipeline::FPanoramicPane& Pane, TUniqueFunction<void(FLinearColor*, FIntPoint)> OnDebugSampleAvailable)
+{
+
+	TaskConcurrencyLimiter->Push(UE_SOURCE_LOCATION, [
+		this,
+		Pane, 
+		InData = MoveTemp(InData),
+		OnDebugSampleAvailable = MoveTemp(OnDebugSampleAvailable)](uint32 Slot)
+		{
+			SCOPE_CYCLE_COUNTER(STAT_MoviePipeline_PanoBlend);
+			const float BlendStartTime = FPlatformTime::Seconds();
+
+			// The way blending works is that each sample that comes in gets its own memory to do the blending into. 
+			// We calculate a bounding box for where the data would end up in the resulting final texture map, and then
+			// blend into it. Then after all samples have come in, we can simply add the results from each blended image
+			// together to get our final value, without having a lot of per-pixel contention during blending
+			//
+			// This math below works out the output dimensions for this sample, which will let us calculate the size.
+			FIntPoint SampleSize = Pane.Resolution;
+			FRotator SampleRotation = FRotator(Pane.bUseLocalRotation ? Pane.CameraLocalRotation : Pane.CameraRotation);
+
+			const float SampleHalfHorizontalFoVDegrees = 0.5f * Pane.HorizontalFieldOfView;
+			const float SampleHalfVerticalFoVDegrees = 0.5f * Pane.VerticalFieldOfView; // FMath::RadiansToDegrees((FMath::Atan((0.5f * SampleSize.X / FMath::Tan(DataPayload->Pane.Data.HorizontalFieldOfView / 2.f))))); // CTransformTools::getVFovFromHFovAndARInDegrees(iSampleFovInDegrees, static_cast<float>(iSampleSize.X) / static_cast<float>(iSampleSize.Y));
+			const float SampleHalfHorizontalFoVCosine = FMath::Cos(FMath::DegreesToRadians(SampleHalfHorizontalFoVDegrees));
+			const float SampleHalfVerticalFoVCosine = FMath::Cos(FMath::DegreesToRadians(SampleHalfVerticalFoVDegrees));
+
+			// Now calculate which direction the Panoramic Pane (that this sample represents) was facing originally.
+			const float SampleYawRad = FMath::DegreesToRadians(SampleRotation.Yaw);
+			const float SamplePitchRad = FMath::DegreesToRadians(SampleRotation.Pitch);
+			const FVector SampleDirectionOnTheta = FVector(FMath::Cos(SampleYawRad), FMath::Sin(SampleYawRad), 0);
+			const FVector SampleDirectionOnPhi = FVector(FMath::Cos(SamplePitchRad), 0.f, FMath::Sin(SamplePitchRad));
+
+			// Now construct a projection matrix representing the sample matching the original perspective it was taken from.
+			const FMatrix SampleProjectionMatrix = FReversedZPerspectiveMatrix(FMath::DegreesToRadians(SampleHalfHorizontalFoVDegrees), SampleSize.X, SampleSize.Y, Pane.NearClippingPlane);
+			// For our given output size, figure out how many degrees each pixel represents.
+			const float EquiRectMapThetaStep = 360.f / (float)OutputEquirectangularMapSize.X;
+			const float EquiRectMapPhiStep = 180.f / (float)OutputEquirectangularMapSize.Y;
+
+			// Compute the index bounds in the equirectangular map corresponding to the sample bounds, so we don't loop over unnecessary pixels.
+			// This is approximated according to the weighting function for blending too.
+			// This assumes that the origin of the equirectangular map (0,0) has a yaw/pitch equal to -180/-90.
+			// Phi evolves in the opposite direction of Y (Y's origin is up-left)
+			// Pitch is clamped, because there is no vertical wrapping in the map
+			// Yaw is not clamped, because horizontal wrapping is possible. 
+			// The MinBound for X can actually be greater than the MaxBound due to wrapping, modulo is applied at eval time to ensure it wraps right.
+			const float SampleYawMin = SampleRotation.Yaw - SampleHalfHorizontalFoVDegrees;
+			const float SampleYawMax = SampleRotation.Yaw + SampleHalfHorizontalFoVDegrees;
+			const int32 PixelIndexHorzMinBound = FMath::FloorToInt(((SampleYawMin)+180.f) / EquiRectMapThetaStep);
+			const int32 PixelIndexHorzMaxBound = FMath::FloorToInt(((SampleYawMax)+180.f) / EquiRectMapThetaStep);
+
+			const float SamplePitchMin = FMath::Max(SampleRotation.Pitch - SampleHalfVerticalFoVDegrees, -90.f); // Clamped to [-90, 90]
+			const float SamplePitchMax = FMath::Min(SampleRotation.Pitch + SampleHalfVerticalFoVDegrees, 90.f); // Clamped to [-90, 90]
+			const int32 PixelIndexVertMinBound = FMath::Max((OutputEquirectangularMapSize.Y) - FMath::FloorToInt((SamplePitchMax + 90.f) / EquiRectMapPhiStep), 0);
+			const int32 PixelIndexVertMaxBound = FMath::Min((OutputEquirectangularMapSize.Y) - FMath::FloorToInt((SamplePitchMin + 90.f) / EquiRectMapPhiStep), OutputEquirectangularMapSize.Y);
+
+			// Build a rect that describes which part of the output map we'll be rendering into
+			const FIntPoint OutputBoundsMin = FIntPoint(PixelIndexHorzMinBound, PixelIndexVertMinBound);
+			const FIntPoint OutputBoundsMax = FIntPoint(PixelIndexHorzMaxBound, PixelIndexVertMaxBound);
+
+			const int32 PixelWidth = OutputBoundsMax.X - OutputBoundsMin.X;
+			const int32 PixelHeight = OutputBoundsMax.Y - OutputBoundsMin.Y;
+			FIntPoint Resolution = FIntPoint(PixelWidth, PixelHeight);
+
+			// We need to find a free entry in the pool that matches our desired size. We pool these because
+			// the memory allocation is expensive, and many of them will have similar sizes.
+			FPoolEntry* FoundPool = nullptr;
+
+			{
+				FScopeLock ScopeLock(&PoolAccessMutex);
+
+				// The concurrent task limiter should prevent us from allocating too many pools, so this will slowly
+				// grow to match the max concurrency limiter.
+				for (TUniquePtr<FPoolEntry>& Item : TempBufferPool)
+				{
+					if (!Item->bActive)
+					{
+						FoundPool = Item.Get();
+						//UE_LOG(LogTemp, Log, TEXT("Reused pool entry, size: %d,%d newSize: %d oldSize:%d"), 
+						//	Resolution.X, Resolution.Y, Resolution.X*Resolution.Y, Item->Data.Max());
+						break;
+					}
+				}
+
+				if (!FoundPool)
+				{
+					TUniquePtr<FPoolEntry> NewPoolEntry = MakeUnique<FPoolEntry>();
+					FoundPool = NewPoolEntry.Get();
+					TempBufferPool.Add(MoveTemp(NewPoolEntry));
+					// UE_LOG(LogTemp, Log, TEXT("Allocated new pool entry, size: %d,%d NewPoolCount: %d"), Resolution.X, Resolution.Y, TempBufferPool.Num());
+				}
+
+				// Now that we've either found a pool for reuse we'll initialize enough of it that other threads
+				// don't come in and try and use it, then we'll free the scope lock and finish initializing it.
+				FoundPool->bActive = true;
+				FoundPool->Resolution = Resolution;
+				FoundPool->OutputBoundsMin = OutputBoundsMin;
+				FoundPool->OutputBoundsMax = OutputBoundsMax;
+			}
+
+
+			check(FoundPool);
+			// We ensure we always have room for the resolution as that is what will be blended into it.
+			int32 MaxSizeX = FMath::Max(SampleSize.X, FoundPool->Resolution.X);
+			int32 MaxSizeY = FMath::Max(SampleSize.Y, FoundPool->Resolution.Y);
+			// This should generally avoid reallocations, smaller blended tiles will just use a sub-region of
+			// the memory.
+			FoundPool->Data.SetNumUninitialized(MaxSizeX * MaxSizeY, EAllowShrinking::No);
+
+			// We need to zero-initialize the data in this patch, especially if it was re-used, because we additively
+			// add images together later.
+			FMemory::Memzero(FoundPool->Data.GetData(), FoundPool->Data.Num() * FoundPool->Data.GetTypeSize());
+
+			// Finally we can perform our actual blending. We blend into our intermediate buffer
+			// instead of the final output array to avoid multiple threads contending for pixels.
+			// This uses the resolution of the output (in blended space) as the size, and we pulls
+			// from the appropriate place in the incoming data.
+			ParallelFor(Resolution.Y,
+				[
+					this,
+					PixelIndexVertMinBound, 
+					PixelIndexHorzMinBound,
+					EquiRectMapThetaStep,
+					EquiRectMapPhiStep,
+					SampleDirectionOnPhi,
+					SampleDirectionOnTheta,
+					SampleHalfHorizontalFoVCosine,
+					SampleHalfVerticalFoVCosine,
+					&SampleRotation,
+					&SampleProjectionMatrix,
+					&SampleSize,
+					&Pane,
+					InDataPtr = InData.Get(),
+					FoundPoolPtr = FoundPool
+				](int32 RowY)
+				{
+					SCOPE_CYCLE_COUNTER(STAT_MoviePipeline_PanoBlendActual);
+					for (int32 RowX = 0; RowX < FoundPoolPtr->Resolution.X; RowX++)
+					{
+						int32 Y = PixelIndexVertMinBound + RowY;
+						int32 X = PixelIndexHorzMinBound + RowX;
+
+						// These X, Y coordinates are in output resolution space which is where we want to blend to.
+						// Our X bounds may go OOB, but we wrap horizontally so we need to figure out the proper X index.
+						const int32 OutputPixelX = ((X % OutputEquirectangularMapSize.X) + OutputEquirectangularMapSize.X) % OutputEquirectangularMapSize.X;
+						const int32 OutputPixelY = Y;
+
+						// Get the spherical coordinates (Theta and Phi) corresponding to the X and Y of the equirectangular map coordinates, converted to
+						// [-180, 180] and [-90, 90] coordinate space respectively. The half pixel offset is used to make the center of a pixel be considered
+						// that coordinate, and Phi increments in the opposite direction of Y.
+						const float Theta = EquiRectMapThetaStep * (((float)OutputPixelX) + 0.5f) - 180.f;
+						const float Phi = EquiRectMapPhiStep * (((float)OutputEquirectangularMapSize.Y - OutputPixelY) + 0.5f) - 90.f;
+						// Now convert the spherical coordinates into an actual direction (on the output map)
+						const float ThetaDeg = FMath::DegreesToRadians(Theta);
+						const float PhiDeg = FMath::DegreesToRadians(Phi);
+						const FVector OutputDirection(FMath::Cos(PhiDeg) * FMath::Cos(ThetaDeg), FMath::Cos(PhiDeg) * FMath::Sin(ThetaDeg), FMath::Sin(PhiDeg));
+						const FVector OutputDirectionTheta = FVector(FMath::Cos(ThetaDeg), FMath::Sin(ThetaDeg), 0);
+						const FVector OutputDirectionPhi = FVector(FMath::Cos(PhiDeg), 0.f, FMath::Sin(PhiDeg));
+
+						// Now we can compute how much the sample should influence this pixel. It is weighted by angular distance to the direction
+						// so that the edges have less influence (where they'd be more distorted anyways).
+						const float DirectionPhiDot = FVector::DotProduct(OutputDirectionPhi, SampleDirectionOnPhi);
+						const float DirectionThetaDot = FVector::DotProduct(OutputDirectionTheta, SampleDirectionOnTheta);
+
+						// The divide is important, as otherwise at large resolutions the individual weights become really small for the whole image.
+						const float WeightTheta = FMath::Max(DirectionThetaDot - SampleHalfHorizontalFoVCosine, 0.0f) / (1.0f - SampleHalfHorizontalFoVCosine);
+						const float WeightPhi = FMath::Max(DirectionPhiDot - SampleHalfVerticalFoVCosine, 0.0f) / (1.0f - SampleHalfVerticalFoVCosine);
+
+						const float SampleWeight = WeightTheta * WeightPhi;
+						const float SampleWeightSquared = SampleWeight * SampleWeight; // Exponential falloff produces a nicer blending result.
+
+						// The sample weight may be very small and not worth influencing this pixel.
+						if (SampleWeightSquared > KINDA_SMALL_NUMBER)
+						{
+							// Transform the direction vector from the equirectangular map world space to sample world space
+							FVector4 DirectionInSampleWorldSpace = FVector4(SampleRotation.UnrotateVector(OutputDirection), 1.0f);
+
+							static const FMatrix UnrealCoordinateConversion = FMatrix(
+								FPlane(0, 0, 1, 0),
+								FPlane(1, 0, 0, 0),
+								FPlane(0, 1, 0, 0),
+								FPlane(0, 0, 0, 1));
+							DirectionInSampleWorldSpace = UnrealCoordinateConversion.TransformFVector4(DirectionInSampleWorldSpace);
+
+							// Then project that direction into sample clip space
+							FVector4 DirectionInSampleClipSpace = SampleProjectionMatrix.TransformFVector4(DirectionInSampleWorldSpace);
+
+							// Converted into normalized device space (Divide by w for perspective)
+							FVector DirectionInSampleNDSpace = FVector(DirectionInSampleClipSpace) / DirectionInSampleClipSpace.W;
+
+							// Get the final pixel coordinates (direction in screen space)
+							FVector2D DirectionInSampleScreenSpace = ((FVector2D(DirectionInSampleNDSpace) + 1.0f) / 2.0f) * FVector2D(SampleSize.X, SampleSize.Y);
+
+							DirectionInSampleScreenSpace.Y = ((float)SampleSize.Y - DirectionInSampleScreenSpace.Y) - 1.0f;
+
+							// Do a bilinear color sample at the pixel coordinates (from the sample), weight it, and add it to the output map.
+							bool bClipped = false;
+							FLinearColor SampleColor;
+							if (Pane.FilterType == EMoviePipelinePanoramicFilterType::Bilinear)
+							{
+								SampleColor = GetColorBilinearFiltered(InDataPtr, DirectionInSampleScreenSpace, bClipped, true);
+							}
+							else
+							{
+								UE::MoviePipeline::FPanoramicPane::FCubicInterpolationParams CubicParams = UE::MoviePipeline::FPanoramicPane::FCubicInterpolationParams::GetParamsForType(Pane.FilterType);
+								SampleColor = GetColorCubicFiltered(InDataPtr, DirectionInSampleScreenSpace, CubicParams.ParamB, CubicParams.ParamC, bClipped, true);
+							}
+
+							if (!bClipped)
+							{
+								// When we calculate the actual output location we need to shift the X/Y. This is because up until now the math has been done in
+								// output resolution space, but each sample only allocates a color map big enough for itself. It'll get shifted back out to the
+								// right location later.
+								int32 SampleOutputX = OutputPixelX - FoundPoolPtr->OutputBoundsMin.X;
+								// Mod this again by our output map so we don't OOB on it. It'll wrap weirdly in the output map but should restore fine.
+								SampleOutputX = ((SampleOutputX % (FoundPoolPtr->Resolution.X)) + (FoundPoolPtr->Resolution.X)) % (FoundPoolPtr->Resolution.X); // Positive Mod
+								int32 SampleOutputY = Y;
+								SampleOutputY -= FoundPoolPtr->OutputBoundsMin.Y;
+
+								const int32 FinalIndex = SampleOutputX + (SampleOutputY * (FoundPoolPtr->Resolution.X));
+
+								FoundPoolPtr->Data[FinalIndex] += SampleColor * SampleWeightSquared;
+							}
+						}
+					}
+				});
+
+			// We don't have a great way to make this abstract between the two systems (as they have different payloads they want in the output merger)
+			// and we need the work to be done mid-cycle before this function returns, so we call an event and let them COPY the data if they want, and
+			// then we assume the data still exist and blend it into our output image before we release the pooled data.
+			// Once we have finished doing the blending, we optionally pass it along to the output merger as a debug sample.
+			if (OnDebugSampleAvailable.IsSet())
+			{
+				OnDebugSampleAvailable(FoundPool->Data.GetData(), FoundPool->Resolution);
+			}
+
+			// Now that blending is complete on this sample, we can place it in the output map. We want this to be fast, because other threads are probably waiting
+			// to do the same thing, but we want to do this as soon as possible to return the blended-space array to the pool so that subsequent samples could
+			// potentially reuse it within the same frame.
+			{
+				FScopeLock ScopeLock(&OutputMapAccessMutex);
+				ParallelFor(FoundPool->Resolution.Y, 
+					[this, FoundPoolPtr = FoundPool] (int32 SampleY) {
+					SCOPE_CYCLE_COUNTER(STAT_MoviePipeline_PanoBlendJoin);
+
+					for (int32 SampleX = 0; SampleX < FoundPoolPtr->Resolution.X; SampleX++)
+					{
+						int32 OriginalX = SampleX + FoundPoolPtr->OutputBoundsMin.X;
+						int32 OriginalY = SampleY + FoundPoolPtr->OutputBoundsMin.Y;
+						const int32 OutputPixelX = ((OriginalX % OutputEquirectangularMapSize.X) + OutputEquirectangularMapSize.X) % OutputEquirectangularMapSize.X;
+						const int32 OutputPixelY = OriginalY;
+
+						int32 SourceIndex = SampleX + (SampleY * (FoundPoolPtr->Resolution.X));
+						int32 DestIndex = OutputPixelX + (OutputPixelY * OutputEquirectangularMapSize.X);
+						OutputEquirectangularMap[DestIndex] += FoundPoolPtr->Data[SourceIndex];
+					}
+					});
+			}
+
+			// Finally we release the pool entry back to the pool.
+			{
+				// We don't take the scope lock because there can be threads above in a busy loop waiting on the lock,
+				// and we need to release this entry back to the pool so they can consume it. bActive is guarded an
+				// atomic so it's okay to write to it even as other threads are trying to read to see if it's no longer active.
+				FoundPool->bActive = false;
+
+				// We no longer have control over this pool entry as we've released it to the pool, null it out
+				// so it doesn't accidentally get used.
+				FoundPool = nullptr;
+			}
+		});
+}
+
+void FMoviePipelinePanoramicBlenderBase::FetchFinalPixelDataHalfFloat(TArray64<FFloat16Color>& OutPixelData) const
+{
+	OutPixelData.SetNumUninitialized(OutputEquirectangularMap.Num());
+
+	ParallelFor(OutputEquirectangularMapSize.Y, [&](int32 InIndexY)
+		{
+			for (int32 FullX = 0L; FullX < OutputEquirectangularMapSize.X; FullX++)
+			{
+				// Be careful with this index, make sure to use 64bit math, not 32bit
+				int64 Index = (int64(InIndexY) * int64(OutputEquirectangularMapSize.X)) + int64(FullX);
+				FLinearColor Color = OutputEquirectangularMap[Index];
+
+				// Normalize by the weight in the alpha channel
+				Color.R /= Color.A;
+				Color.G /= Color.A;
+				Color.B /= Color.A;
+				Color.A = 1.f;
+
+				OutPixelData[Index] = FFloat16Color(Color);
+			}
+		});
+}
+
+void FMoviePipelinePanoramicBlenderBase::FetchFinalPixelDataLinearColor(TArray64<FLinearColor>& OutPixelData) const
+{
+	OutPixelData.SetNumUninitialized(OutputEquirectangularMap.Num());
+
+	ParallelFor(OutputEquirectangularMapSize.Y, [&](int32 InIndexY)
+		{
+			for (int32 FullX = 0L; FullX < OutputEquirectangularMapSize.X; FullX++)
+			{
+				// Be careful with this index, make sure to use 64bit math, not 32bit
+				int64 Index = (int64(InIndexY) * int64(OutputEquirectangularMapSize.X)) + int64(FullX);
+				FLinearColor Color = OutputEquirectangularMap[Index];
+
+				// Normalize by the weight in the alpha channel
+				Color.R /= Color.A;
+				Color.G /= Color.A;
+				Color.B /= Color.A;
+				Color.A = 1.f;
+
+				OutPixelData[Index] = Color;
+			}
+		});
+}
+	
+} // UE::MoviePipeline 
diff --git a/Engine/Plugins/MovieScene/MovieRenderPipeline/Source/MovieRenderPipelineCore/Public/MoviePipelinePanoramicBlenderBase.h b/Engine/Plugins/MovieScene/MovieRenderPipeline/Source/MovieRenderPipelineCore/Public/MoviePipelinePanoramicBlenderBase.h
new file mode 100644
--- /dev/null
+++ b/Engine/Plugins/MovieScene/MovieRenderPipeline/Source/MovieRenderPipelineCore/Public/MoviePipelinePanoramicBlenderBase.h
@@ -0,0 +1,156 @@
+// Copyright Epic Games, Inc. All Rights Reserved.
+#pragma once
+
+#include "MovieRenderPipelineDataTypes.h"
+#include "Tasks/TaskConcurrencyLimiter.h"
+#include "MoviePipelinePanoramicBlenderBase.generated.h"
+
+// Forward Declares
+struct FImagePixelData;
+
+UENUM(BlueprintType)
+enum class EMoviePipelinePanoramicFilterType : uint8
+{
+	/** 2x2 Bilinear Interpolation. Fastest, nearly the same quality as other options. */
+	Bilinear,
+	/** Cubic Catmull-Rom interpolation. Slightly sharper than other results. Uses B=0, C=1/2 in parameterized cubic equation. */
+	Catmull,
+	/** Cubic Mitchell-Netravali interpolation. More neutral look. Uses B=0.33, C=0.33 in parameterized cubic equation. */
+	Mitchell
+};
+
+namespace UE::MoviePipeline
+{
+	struct FPanoramicPane
+	{
+		// The camera location as defined by the actual sequence, consistent for all panes.
+		FVector OriginalCameraLocation;
+		// The camera location last frame, used to ensure camera motion vectors are right.
+		FVector PrevOriginalCameraLocation;
+		// The camera rotation as defined by the actual sequence
+		FRotator OriginalCameraRotation;
+		// The camera rotation last frame, used to ensure camera motion vectors are right.
+		FRotator PrevOriginalCameraRotation;
+		// The near clip plane distance from the camera.
+		float NearClippingPlane;
+
+		// How far apart are the eyes (total) for stereo?
+		float EyeSeparation;
+		float EyeConvergenceDistance;
+
+		// The horizontal field of view this pane was rendered with
+		float HorizontalFieldOfView;
+		float VerticalFieldOfView;
+
+		FIntPoint Resolution;
+
+		// The actual rendering location for this pane, offset by the stereo eye if needed.
+		FVector CameraLocation;
+		FVector PrevCameraLocation;
+		FRotator CameraRotation;
+		FRotator CameraLocalRotation;
+		FRotator PrevCameraRotation;
+
+		// If true, uses only the CameraLocalRotation which means that if the camera yaws, so will the resulting blended image.
+		bool bUseLocalRotation;
+
+		// How many horizontal segments are there total.
+		int32 NumHorizontalSteps;
+		int32 NumVerticalSteps;
+
+		// Which horizontal segment are we?
+		int32 HorizontalStepIndex;
+		// Which vertical segment are we?
+		int32 VerticalStepIndex;
+
+		struct FCubicInterpolationParams
+		{
+			float ParamB;
+			float ParamC;
+
+			static FCubicInterpolationParams GetParamsForType(const EMoviePipelinePanoramicFilterType InType)
+			{
+				FCubicInterpolationParams Params;
+				switch (InType)
+				{
+				case EMoviePipelinePanoramicFilterType::Catmull:
+					Params.ParamB = 0.0f;
+					Params.ParamC = 0.5f;
+					break;
+				case EMoviePipelinePanoramicFilterType::Mitchell:
+					Params.ParamB = 0.33f;
+					Params.ParamC = 0.33f;
+					break;
+				default:
+					ensureMsgf(false, TEXT("GetParamsForType shouldn't be called for non-cubic interpolations."));
+					Params.ParamB = 0.0f;
+					Params.ParamC = 0.0f;
+					break;
+				}
+				return Params;
+			}
+
+		};
+		
+		EMoviePipelinePanoramicFilterType FilterType = EMoviePipelinePanoramicFilterType::Bilinear;
+
+		// When indexing into arrays of Panes, which index is this?
+		int32 GetAbsoluteIndex() const
+		{
+			const int32 EyeOffset = EyeIndex == -1 ? 0 : EyeIndex;
+			const int32 NumEyeRenders = EyeIndex == -1 ? 1 : 2;
+			return  (VerticalStepIndex * NumHorizontalSteps * NumEyeRenders) + HorizontalStepIndex + EyeOffset;
+		}
+
+		// -1 if no stereo, 0 left eye, 1 right eye.
+		int32 EyeIndex;
+	};
+		
+	/**
+	* This class is responsible for blending a together a single Equirectangular Image from a series of individual
+	* renders, and can blend multiple samples in a threadsafe way. For each incoming sample the pixel data should
+	* contain information about panoramic pane (orientation, index, etc.), and one all samples have been called
+	* BlendSample_AnyThread, it is safe for the owner of this instance to fetch the image data, which returns a
+	* copy of the output image. You can then call Initialize() on it to reset the output image without reallocating
+	* memory, which allows for reusing a given output blender for multiple frames.
+	* 
+	* This implementation works by allocating memory for each incoming sample that is the size of the data
+	* once blended. Depending on where in the projection it is, different samples will take up different parts
+	* of the output image (with different resolutions) so the pool stores available buffers by resolution.
+	* Once a pool is either found or allocated for the sample, the incoming data is read from. Instead of taking
+	* each sample in the incoming data and figuring out where it would go in the output image, we instead work backwards,
+	* and calculate from each output pixel in the range that the sample would affect, sample the source image with filtering.
+	* 
+	* Once the blending into the temporary buffer is complete, a lock on the output array is taken and the data is added to
+	* the output, and the buffer is returned to the pool.
+	*/
+	class MOVIERENDERPIPELINECORE_API FMoviePipelinePanoramicBlenderBase
+	{
+	public:
+		void Initialize(const FIntPoint InOutputResolution);
+	public:
+		void BlendSample_AnyThread(TUniquePtr<FImagePixelData>&& InData, const UE::MoviePipeline::FPanoramicPane& Pane, TUniqueFunction<void(FLinearColor*, FIntPoint)> OnDebugSampleAvailable);
+		void FetchFinalPixelDataHalfFloat(TArray64<FFloat16Color>& OutPixelData) const;
+		void FetchFinalPixelDataLinearColor(TArray64<FLinearColor>& OutPixelData) const;
+
+		TUniquePtr<UE::Tasks::FTaskConcurrencyLimiter> TaskConcurrencyLimiter;
+	private:
+		struct FPoolEntry
+		{
+			std::atomic<bool> bActive;
+			TArray64<FLinearColor> Data;
+
+			FIntPoint Resolution;
+			FIntPoint OutputBoundsMin;
+			FIntPoint OutputBoundsMax;
+		};
+
+		TArray64<TUniquePtr<FPoolEntry>> TempBufferPool;
+		FCriticalSection PoolAccessMutex;
+		FCriticalSection OutputMapAccessMutex;
+
+		TArray64<FLinearColor> OutputEquirectangularMap;
+		FIntPoint OutputEquirectangularMapSize;
+
+	};
+}
\ No newline at end of file
diff --git a/Engine/Plugins/MovieScene/MovieRenderPipeline/Source/MovieRenderPipelineCore/Public/MoviePipelineUtils.h b/Engine/Plugins/MovieScene/MovieRenderPipeline/Source/MovieRenderPipelineCore/Public/MoviePipelineUtils.h
--- a/Engine/Plugins/MovieScene/MovieRenderPipeline/Source/MovieRenderPipelineCore/Public/MoviePipelineUtils.h
+++ b/Engine/Plugins/MovieScene/MovieRenderPipeline/Source/MovieRenderPipelineCore/Public/MoviePipelineUtils.h
@@ -1,18 +1,19 @@
 // Copyright Epic Games, Inc. All Rights Reserved.
 #pragma once
 
+#include "CineCameraComponent.h"
 #include "CoreMinimal.h"
-#include "Engine/Scene.h"
 #include "Engine/Engine.h"
-#include "MovieRenderPipelineDataTypes.h"
+#include "Engine/Scene.h"
+#include "Evaluation/MovieSceneSequenceHierarchy.h"
 #include "Evaluation/MovieSceneTimeTransform.h"
-#include "Misc/FrameNumber.h"
 #include "Math/Range.h"
+#include "Misc/FrameNumber.h"
 #include "Misc/FrameRate.h"
-#include "Evaluation/MovieSceneSequenceHierarchy.h"
-#include "MovieSceneSequenceID.h"
-#include "CineCameraComponent.h"
+#include "MoviePipelinePanoramicBlenderBase.h"
 #include "MoviePipelineQueue.h"
+#include "MovieRenderPipelineDataTypes.h"
+#include "MovieSceneSequenceID.h"
 
 // Forward Declare
 class UClass;
@@ -162,8 +163,11 @@ namespace UE
 {
 	namespace MovieRenderPipeline
 	{
+		/** Finds UMoviePipelineSetting classes declared in the project using the asset registry, potentially including blueprint implemented classes. */
 		MOVIERENDERPIPELINECORE_API TArray<UClass*> FindMoviePipelineSettingClasses(UClass* InBaseClass, const bool bIncludeBlueprints = true);
+		/** Anti-aliasing is defined as a Project Setting, and then MRG can optionally override it. Utility function for getting the project setting and optionally applying the MRQ override to the return value. */
 		MOVIERENDERPIPELINECORE_API EAntiAliasingMethod GetEffectiveAntiAliasingMethod(bool bOverride, EAntiAliasingMethod OverrideMethod);
+		/** Anti-aliasing is defined as a Project Setting, and then MRQ can optionally override it. Utility function for getting the project setting and optionally applying the MRQ override to the return value. */
 		MOVIERENDERPIPELINECORE_API EAntiAliasingMethod GetEffectiveAntiAliasingMethod(const UMoviePipelineAntiAliasingSetting* InSetting);
 		
 		UE_DEPRECATED(5.3, "Do not use, this is here as a temporary workaround for another issue.")
@@ -211,8 +215,11 @@ namespace UE
 		MOVIERENDERPIPELINECORE_API FIntPoint ScaleResolutionByOverscan(const float OverscanPercentage, const FIntPoint& InOutputResolution);
 		namespace Panoramic
 		{
+			/** Returns an array of values distributed throughout the given range, used to calculate which angles the Panoramic renderers should orient the camera to. */
 			MOVIERENDERPIPELINECORE_API TArray<float> DistributeValuesInInterval(float InMin, float InMax, int32 InNumDivisions, bool bInInclusiveMax);
+			/** Given the primary camera rotation and a pane describing the current pane, calculate a new rotation / location for the camera. If stereo was supported, the location would change with the rotation. */
 			MOVIERENDERPIPELINECORE_API void GetCameraOrientationForStereo(FVector& OutLocation, FRotator& OutRotation, FRotator& OutLocalRotation, const UE::MoviePipeline::FPanoramicPane& InPane, const int32 InStereoIndex, const bool bInPrevPosition);
+			/** Auto-exposure cubemap passes are laid out in a 3x2 arrangement, and rounded down to a multiple of 8 pixels. */
 			MOVIERENDERPIPELINECORE_API int32 ComputeAutoExposureCubeCaptureSize(FIntPoint Resolution);
 		}
 	}
diff --git a/Engine/Plugins/MovieScene/MovieRenderPipeline/Source/MovieRenderPipelineCore/Public/MovieRenderPipelineDataTypes.h b/Engine/Plugins/MovieScene/MovieRenderPipeline/Source/MovieRenderPipelineCore/Public/MovieRenderPipelineDataTypes.h
--- a/Engine/Plugins/MovieScene/MovieRenderPipeline/Source/MovieRenderPipelineCore/Public/MovieRenderPipelineDataTypes.h
+++ b/Engine/Plugins/MovieScene/MovieRenderPipeline/Source/MovieRenderPipelineCore/Public/MovieRenderPipelineDataTypes.h
@@ -1384,60 +1384,6 @@ namespace UE
 			TSubclassOf<UMovieRenderDebugWidget> DebugWidgetClass;
 			bool bRenderViewport = false;
 		};
-
-		struct FPanoramicPane
-		{
-			// The camera location as defined by the actual sequence, consistent for all panes.
-			FVector OriginalCameraLocation;
-			// The camera location last frame, used to ensure camera motion vectors are right.
-			FVector PrevOriginalCameraLocation;
-			// The camera rotation as defined by the actual sequence
-			FRotator OriginalCameraRotation;
-			// The camera rotation last frame, used to ensure camera motion vectors are right.
-			FRotator PrevOriginalCameraRotation;
-			// The near clip plane distance from the camera.
-			float NearClippingPlane;
-
-			// How far apart are the eyes (total) for stereo?
-			float EyeSeparation;
-			float EyeConvergenceDistance;
-
-			// The horizontal field of view this pane was rendered with
-			float HorizontalFieldOfView;
-			float VerticalFieldOfView;
-
-			FIntPoint Resolution;
-
-			// The actual rendering location for this pane, offset by the stereo eye if needed.
-			FVector CameraLocation;
-			FVector PrevCameraLocation;
-			FRotator CameraRotation;
-			FRotator CameraLocalRotation;
-			FRotator PrevCameraRotation;
-
-			// If true, uses only the CameraLocalRotation which means that if the camera yaws, so will the resulting blended image.
-			bool bUseLocalRotation;
-
-			// How many horizontal segments are there total.
-			int32 NumHorizontalSteps;
-			int32 NumVerticalSteps;
-
-			// Which horizontal segment are we?
-			int32 HorizontalStepIndex;
-			// Which vertical segment are we?
-			int32 VerticalStepIndex;
-
-			// When indexing into arrays of Panes, which index is this?
-			int32 GetAbsoluteIndex() const
-			{
-				const int32 EyeOffset = EyeIndex == -1 ? 0 : EyeIndex;
-				const int32 NumEyeRenders = EyeIndex == -1 ? 1 : 2;
-				return  (VerticalStepIndex * NumHorizontalSteps * NumEyeRenders) + HorizontalStepIndex + EyeOffset;
-			}
-
-			// -1 if no stereo, 0 left eye, 1 right eye.
-			int32 EyeIndex;
-		};
 	}
 }
 /**
diff --git a/Engine/Plugins/MovieScene/MovieRenderPipeline/Source/MovieRenderPipelineRenderPasses/Private/Graph/Nodes/MovieGraphDeferredPanoramicPassNode.cpp b/Engine/Plugins/MovieScene/MovieRenderPipeline/Source/MovieRenderPipelineRenderPasses/Private/Graph/Nodes/MovieGraphDeferredPanoramicPassNode.cpp
--- a/Engine/Plugins/MovieScene/MovieRenderPipeline/Source/MovieRenderPipelineRenderPasses/Private/Graph/Nodes/MovieGraphDeferredPanoramicPassNode.cpp
+++ b/Engine/Plugins/MovieScene/MovieRenderPipeline/Source/MovieRenderPipelineRenderPasses/Private/Graph/Nodes/MovieGraphDeferredPanoramicPassNode.cpp
@@ -4,12 +4,28 @@
 #include "Graph/Renderers/MovieGraphDeferredPanoramicPass.h"
 #include "MoviePipelineTelemetry.h"
 
+UMovieGraphDeferredPanoramicNode::UMovieGraphDeferredPanoramicNode()
+	: NumHorizontalSteps(8)
+	, NumVerticalSteps(3)
+	, bFollowCameraOrientation(true)
+	, bAllocateHistoryPerPane(false)
+	, bPageToSystemMemory(false)
+	, SpatialSampleCount(1)
+	, AntiAliasingMethod(EAntiAliasingMethod::AAM_TSR)
+	, Filter(EMoviePipelinePanoramicFilterType::Bilinear)
+	, bWriteAllSamples(false)
+	, bDisableToneCurve(false)
+	, bAllowOCIO(true)
+	, ViewModeIndex(VMI_Lit)
+{}
+
 FEngineShowFlags UMovieGraphDeferredPanoramicNode::GetShowFlags() const
 {
-	FEngineShowFlags Flags(ESFIM_Game);
+	FEngineShowFlags Flags = Super::GetShowFlags();
 	Flags.SetVignette(false);
 	Flags.SetSceneColorFringe(false);
 	Flags.SetPhysicalMaterialMasks(false);
+	Flags.SetDepthOfField(false);
 
 	return Flags;
 }
diff --git a/Engine/Plugins/MovieScene/MovieRenderPipeline/Source/MovieRenderPipelineRenderPasses/Private/Graph/Renderers/MovieGraphDeferredPanoramicPass.cpp b/Engine/Plugins/MovieScene/MovieRenderPipeline/Source/MovieRenderPipelineRenderPasses/Private/Graph/Renderers/MovieGraphDeferredPanoramicPass.cpp
--- a/Engine/Plugins/MovieScene/MovieRenderPipeline/Source/MovieRenderPipelineRenderPasses/Private/Graph/Renderers/MovieGraphDeferredPanoramicPass.cpp
+++ b/Engine/Plugins/MovieScene/MovieRenderPipeline/Source/MovieRenderPipelineRenderPasses/Private/Graph/Renderers/MovieGraphDeferredPanoramicPass.cpp
@@ -10,6 +10,7 @@
 #include "Graph/Nodes/MovieGraphApplyViewportLookNode.h"
 #include "Graph/Nodes/MovieGraphCameraNode.h"
 #include "Graph/Nodes/MovieGraphDeferredPanoramicPassNode.h"
+#include "MoviePipelinePanoramicBlenderBase.h"
 #include "MoviePipelineQueue.h"
 #include "MoviePipelineUtils.h"
 #include "MovieRenderOverlappedImage.h"
@@ -18,39 +19,245 @@
 
 namespace UE::MovieGraph::Rendering
 {
-FMovieGraphOutputMergerFrame FMovieGraphPanoramicBlender::MovieGraphDummyOutputFrame;
-TQueue<FMovieGraphOutputMergerFrame> FMovieGraphPanoramicBlender::MovieGraphDummyOutputQueue;
+	struct FMovieGraphPanoSampleState : public FMovieGraphSampleState
+	{
+		virtual TSharedRef<FMovieGraphSampleState> Copy() const override
+		{
+			return MakeShared<FMovieGraphPanoSampleState>(*this);
+		}
+
+		UE::MoviePipeline::FPanoramicPane Pane;
+	};
+
+	static_assert(std::is_base_of<FMovieGraphSampleState, FMovieGraphPanoSampleState>::value,
+		"PanoSampleState must inherit from FMovieGraphSampleState due to shared basecode that does static casting.");
+
+	class FMovieGraphPanoramicBlender : public UE::MovieGraph::IMovieGraphOutputMerger
+	{
+	public:
+		FMovieGraphPanoramicBlender(TSharedPtr<UE::MovieGraph::IMovieGraphOutputMerger> InOutputMerger, const FIntPoint InOutputResolution)
+		{
+			OutputMerger = InOutputMerger;
+			OutputResolution = InOutputResolution;
+		}
+
+		static FMovieGraphOutputMergerFrame MovieGraphDummyOutputFrame;
+		static TQueue<FMovieGraphOutputMergerFrame> MovieGraphDummyOutputQueue;
+
+		virtual FMovieGraphOutputMergerFrame& AllocateNewOutputFrame_GameThread(const int32 InRenderedFrameNumber) override
+		{
+			// Unsupported, the main Output Builder should be the one tracking this.
+			// Since these are references we have to return something, so we return some dummy data.
+			check(0);
+			return MovieGraphDummyOutputFrame;
+		}
+
+		virtual FMovieGraphOutputMergerFrame& GetOutputFrame_GameThread(const int32 InRenderedFrameNumber) override
+		{
+			// Unsupported, the main Output Builder should be the one tracking this.
+			check(0);
+			return MovieGraphDummyOutputFrame;
+		}
+		virtual void OnCompleteRenderPassDataAvailable_AnyThread(TUniquePtr<FImagePixelData>&& InData)
+		{
+
+			// Copy the payload at the start, because it will get destroyed when the pixel data is moved to a task, which then
+			// uses it and discards it before this function finishes.
+			TSharedRef<FMovieGraphPanoSampleState> DataPayload = StaticCastSharedRef<FMovieGraphPanoSampleState>(InData->GetPayload<FMovieGraphPanoSampleState>()->Copy());
+			
+			// This function is called every time a sample comes in from the GPU (after being accumulated) and needs to handle
+			// multiple samples from multiple frames being in flight at once. First step is to search to see if we're already
+			// working on an output frame for this sample.
+			FPoolEntry* TargetBlender = nullptr;
+			{
+				// Do a quick lock while we're iterating/adding to the PendingData array so a second sample
+				// doesn't come in mid iteration.
+				FScopeLock ScopeLock(&GlobalQueueDataMutex);
+
+				for (TUniquePtr<FPoolEntry>& Item : PendingData)
+				{
+					if (Item->OutputFrameNumber == DataPayload->TraversalContext.Time.OutputFrameNumber && Item->bActive)
+					{
+						TargetBlender = Item.Get();
+						break;
+					}
+				}
+
+				if (!TargetBlender)
+				{
+					// UE_LOG(LogMovieRenderPipeline, Log, TEXT("Starting a new Output Frame in the Panoramic Blender for frame: %d"), DataPayload->SampleState.OutputState.OutputFrameNumber);
+
+					// If we didn't find a blender already working on this frame, we'll try to re-use a previously allocated blender.
+					for (TUniquePtr<FPoolEntry>& Item : PendingData)
+					{
+						if (!Item->bActive)
+						{
+							TargetBlender = Item.Get();
+						}
+					}
+
+					// If we still don't have a target blender, then this is a new one and we need to allocate an entry.
+					if (!TargetBlender)
+					{
+						TUniquePtr<FPoolEntry> NewEntry = MakeUnique<FPoolEntry>();
+						int32 NewIndex = PendingData.Add(MoveTemp(NewEntry));
+						TargetBlender = PendingData[NewIndex].Get();
+					}
+
+					check(TargetBlender);
+
+					// If we were already working on this frame, the first for loop through the pending data would have found it.
+					// So we know that if we get here, that we need to initialize whatever blender we ended up with.
+					TargetBlender->OutputFrameNumber = DataPayload->TraversalContext.Time.OutputFrameNumber;
+					TargetBlender->bActive = true;
+					TargetBlender->NumCompletedAccumulations = 0;
+					TargetBlender->Blender.Initialize(OutputResolution);
+				}
+			}
+
+			// This can get called later (due to blending being async) so only capture by value.
+			auto OnDebugSampleAvailable = [
+			DataPayloadCopy = DataPayload->Copy(),
+			WeakOutputMerger = OutputMerger](FLinearColor* Data, FIntPoint Resolution)
+				{
+					TSharedRef<FMovieGraphPanoSampleState> PayloadAsPano = StaticCastSharedRef<FMovieGraphPanoSampleState>(DataPayloadCopy);
+					if(!PayloadAsPano->bWriteSampleToDisk)
+					{
+						return;
+					}
+
+					if (PayloadAsPano->Pane.EyeIndex >= 0)
+					{
+						PayloadAsPano->Debug_OverrideFilename = FString::Printf(TEXT("/%s_PaneX_%d_PaneY_%dEye_%d-Blended.%d"),
+							*PayloadAsPano->TraversalContext.RenderDataIdentifier.LayerName, PayloadAsPano->Pane.HorizontalStepIndex,
+							PayloadAsPano->Pane.VerticalStepIndex, PayloadAsPano->Pane.EyeIndex, PayloadAsPano->TraversalContext.Time.OutputFrameNumber);
+					}
+					else
+					{
+						PayloadAsPano->Debug_OverrideFilename = FString::Printf(TEXT("/%s_PaneX_%d_PaneY_%d-Blended.%d"),
+							*PayloadAsPano->TraversalContext.RenderDataIdentifier.LayerName, PayloadAsPano->Pane.HorizontalStepIndex,
+							PayloadAsPano->Pane.VerticalStepIndex, PayloadAsPano->TraversalContext.Time.OutputFrameNumber);
+					}
+					
+					// We have to copy the memory because the blender is going to re-use it.
+					TArray64<FLinearColor> BlendDataCopy = TArray64<FLinearColor>(Data, Resolution.X * Resolution.Y);
+					TUniquePtr<TImagePixelData<FLinearColor>> FinalPixelData = MakeUnique<TImagePixelData<FLinearColor>>(Resolution, MoveTemp(BlendDataCopy), PayloadAsPano->Copy());
+					
+					if (ensure(WeakOutputMerger.IsValid()))
+					{
+						WeakOutputMerger.Pin()->OnSingleSampleDataAvailable_AnyThread(MoveTemp(FinalPixelData));
+					}
+				};
+
+			// Now that we know which blender we're trying to accumulate to, we can just send the data to it directly. We're already
+			// on a task thread, and the blending process supports multiple task threads working on blending at the same time.
+			TargetBlender->Blender.BlendSample_AnyThread(MoveTemp(InData), DataPayload->Pane, OnDebugSampleAvailable);
+			
+			// Checking to see if this is the last sample is slightly complicated, because we can have multiple threads in this function at the same time.
+			// Inside the blender, it only lets one thread increment the sample count at the same time, but that means when we look at it, we need to
+			// go through the same lock so that we don't have two threads (in this function) read the value, and both decide they're the last sample.
+			
+			{
+				// We put this behind our lock just so that we don't have two threads get the same value one after the other and still decide
+				// they're the last.
+				FScopeLock ScopeLock(&GlobalQueueDataMutex);
+				TargetBlender->NumCompletedAccumulations++;
+			
+				int32 NumCompletedAccumulations = TargetBlender->NumCompletedAccumulations;
+				const bool bLastSample = NumCompletedAccumulations == DataPayload->Pane.NumHorizontalSteps * DataPayload->Pane.NumVerticalSteps;
+			
+				if (bLastSample)
+				{
+					// BlendSample_AnyThread returns immediately and we'll increment it as completed, so if this is the last sample,
+					// we'll wait for the outstanding work to finish.
+					{
+						TargetBlender->Blender.TaskConcurrencyLimiter->Wait();
+					}
+
+					if (ensure(OutputMerger.IsValid()))
+					{
+						TUniquePtr<TImagePixelData<FLinearColor> > FinalPixelData = MakeUnique<TImagePixelData<FLinearColor>>(OutputResolution, DataPayload->Copy());
+						TargetBlender->Blender.FetchFinalPixelDataLinearColor(FinalPixelData->Pixels);
+			
+						OutputMerger.Pin()->OnCompleteRenderPassDataAvailable_AnyThread(MoveTemp(FinalPixelData));
+					}
+			
+					// Release the pool item so future frames can use it.
+					TargetBlender->bActive = false;
+				}
+			}
+		}
+		virtual void OnSingleSampleDataAvailable_AnyThread(TUniquePtr<FImagePixelData>&& InData) override
+		{
+			// This is used for debug output, just pass it straight through.
+			ensure(OutputMerger.IsValid());
+			OutputMerger.Pin()->OnSingleSampleDataAvailable_AnyThread(MoveTemp(InData));
+		}
+		virtual void AbandonOutstandingWork() override
+		{
+			// Not implemented
+			check(0);
+		}
+		virtual int32 GetNumOutstandingFrames() const
+		{
+			// Not implemented (this function isn't called anywhere right now)
+			return 0;
+		}
+		virtual TQueue<FMovieGraphOutputMergerFrame>& GetFinishedFrames()
+		{
+			// Not implemented
+			check(0);
+			return MovieGraphDummyOutputQueue;
+		}
+
+	private:
+		struct FPoolEntry
+		{
+			UE::MoviePipeline::FMoviePipelinePanoramicBlenderBase Blender;
+			bool bActive;
+			int32 OutputFrameNumber;
+			std::atomic<int32> NumCompletedAccumulations;
+		};
+
+		// Pool entries are allocated as pointers on the heap so that if the array is resized while a thread is
+		// working on a previous frame, it doesn't have the memory moved out from underneath it.
+		TArray<TUniquePtr<FPoolEntry>> PendingData;
+
+		FCriticalSection GlobalQueueDataMutex;
+
+		FIntPoint OutputResolution;
+
+		TWeakPtr<UE::MovieGraph::IMovieGraphOutputMerger> OutputMerger;
+	};
+
+	FMovieGraphOutputMergerFrame FMovieGraphPanoramicBlender::MovieGraphDummyOutputFrame;
+	TQueue<FMovieGraphOutputMergerFrame> FMovieGraphPanoramicBlender::MovieGraphDummyOutputQueue;
 
 FMovieGraphDeferredPanoramicPass::FMovieGraphDeferredPanoramicPass()
 {
 }
 
-namespace
+static void GetFieldOfView(float& OutHorizontal, float& OutVertical)
 {
-	void GetFieldOfView(float& OutHorizontal, float& OutVertical)
-	{
-		// Hard-coded for the moment as we don't support stereo or allowing users to override the pane FOV
-		OutHorizontal = 90.f;
-		OutVertical = 90.f;
-	}
-
-	FIntPoint GetPaneResolution(const FIntPoint& InSize)
-	{
-		// We calculate a different resolution than the final output resolution.
-		float HorizontalFoV;
-		float VerticalFoV;
-		GetFieldOfView(HorizontalFoV, VerticalFoV);
-
-		// Horizontal FoV is a proportion of the global horizontal resolution
-		// ToDo: We might have to check which is higher, if numVerticalPanes > numHorizontalPanes this math might be backwards.
-		float HorizontalRes = (HorizontalFoV / 360.0f) * InSize.X;
-		float Intermediate = FMath::Tan(FMath::DegreesToRadians(VerticalFoV) * 0.5f) / FMath::Tan(FMath::DegreesToRadians(HorizontalFoV) * 0.5f);
-		float VerticalRes = HorizontalRes * Intermediate;
-
-		return FIntPoint(FMath::CeilToInt(HorizontalRes), FMath::CeilToInt(VerticalRes));
-	}
-
+	// Hard-coded for the moment as we don't support stereo or allowing users to override the pane FOV
+	OutHorizontal = 90.f;
+	OutVertical = 90.f;
+}
 
+static FIntPoint GetPaneResolution(const FIntPoint& InSize)
+{
+	// We calculate a different resolution than the final output resolution.
+	float HorizontalFoV;
+	float VerticalFoV;
+	GetFieldOfView(HorizontalFoV, VerticalFoV);
+
+	// Horizontal FoV is a proportion of the global horizontal resolution
+	// ToDo: We might have to check which is higher, if numVerticalPanes > numHorizontalPanes this math might be backwards.
+	float HorizontalRes = (HorizontalFoV / 360.0f) * InSize.X;
+	float Intermediate = FMath::Tan(FMath::DegreesToRadians(VerticalFoV) * 0.5f) / FMath::Tan(FMath::DegreesToRadians(HorizontalFoV) * 0.5f);
+	float VerticalRes = HorizontalRes * Intermediate;
+
+	return FIntPoint(FMath::CeilToInt(HorizontalRes), FMath::CeilToInt(VerticalRes));
 }
 
 void FMovieGraphDeferredPanoramicPass::Setup(TWeakObjectPtr<UMovieGraphDefaultRenderer> InRenderer, TWeakObjectPtr<UMovieGraphImagePassBaseNode> InRenderPassNode, const FMovieGraphRenderPassLayerData& InLayer)
@@ -272,8 +479,9 @@ void FMovieGraphDeferredPanoramicPass::Render(const FMovieGraphTraversalContext&
 
 	if (!bHasPrintedRenderingInfo)
 	{
-		UE_LOG(LogMovieRenderPipeline, Log, TEXT("Set-up Deferred Panoramic Renderer: %s Layer: %s PaneResolution: [%d, %d] PaneCount: [%d, %d] bPageToSystemMemory: %d bAutoExposurePass: %d"),
+		UE_LOG(LogMovieRenderPipeline, Log, TEXT("Set-up Deferred Panoramic Renderer: %s Layer: %s OutputRes: [%d, %d] PaneRes: [%d, %d] PaneCount: [%d, %d] bPageToSystemMemory: %d bAutoExposurePass: %d"),
 			*RenderDataIdentifier.RendererName, *RenderDataIdentifier.LayerName,
+			DesiredOutputResolution.X, DesiredOutputResolution.Y,
 			PaneResolution.X, PaneResolution.Y, ParentNodeThisFrame->NumHorizontalSteps, ParentNodeThisFrame->NumVerticalSteps,
 			ParentNodeThisFrame->GetEnablePageToSystemMemory(), AutoExposureViewStates.Num() > 0);
 
@@ -305,7 +513,7 @@ void FMovieGraphDeferredPanoramicPass::Render(const FMovieGraphTraversalContext&
 					Data.PrevOriginalCameraLocation = PrevTransform.GetLocation();
 					Data.PrevOriginalCameraRotation = FRotator(PrevTransform.GetRotation());
 
-					constexpr int32 StereoIndex = 0;
+					constexpr int32 StereoIndex = -1;
 					Data.EyeIndex = StereoIndex;
 					Data.VerticalStepIndex = PaneY;
 					Data.HorizontalStepIndex = PaneX;
@@ -315,6 +523,7 @@ void FMovieGraphDeferredPanoramicPass::Render(const FMovieGraphTraversalContext&
 					Data.EyeConvergenceDistance = 0.f;
 					Data.bUseLocalRotation = ParentNodeThisFrame->bFollowCameraOrientation;
 					Data.Resolution = PaneResolution;
+					Data.FilterType = ParentNodeThisFrame->Filter;
 
 					// The calculations above are for the main camera, now transform this pane's information to be specific to the current pane.
 					bool bInPrevPos = false;
@@ -325,11 +534,8 @@ void FMovieGraphDeferredPanoramicPass::Render(const FMovieGraphTraversalContext&
 					UE::MoviePipeline::Panoramic::GetCameraOrientationForStereo(Data.PrevCameraLocation, Data.PrevCameraRotation, DummyPrevLocalRot, Data, StereoIndex, bInPrevPos);
 
 					GetFieldOfView(Data.HorizontalFieldOfView, Data.VerticalFieldOfView);
-
 				}
 
-
-
 				UTextureRenderTarget2D* RenderTarget = GraphRenderer->GetOrCreateViewRenderTarget(RenderTargetInitParams, RenderDataIdentifier);
 				FRenderTarget* RenderTargetResource = RenderTarget->GameThread_GetRenderTargetResource();
 				check(RenderTargetResource);
@@ -361,7 +567,7 @@ void FMovieGraphDeferredPanoramicPass::Render(const FMovieGraphTraversalContext&
 				CameraInfo.SamplingParams.SpatialSampleIndex = SpatialIndex;
 				CameraInfo.SamplingParams.SpatialSampleCount = NumSpatialSamples;
 				CameraInfo.SamplingParams.SeedOffset = 0;
-				CameraInfo.ProjectionMatrixJitterAmount = FVector2D(0.f, 0.f);
+				CameraInfo.ProjectionMatrixJitterAmount = FVector2D((SpatialShiftAmount.X) * 2.0f / (float)PaneResolution.X, SpatialShiftAmount.Y * -2.0f / (float)PaneResolution.Y);
 				CameraInfo.bUseCameraManagerPostProcess = !bRenderAllCameras;
 				CameraInfo.ViewInfo.ClearOverscan();
 				
@@ -370,6 +576,7 @@ void FMovieGraphDeferredPanoramicPass::Render(const FMovieGraphTraversalContext&
 					CameraInfo.ViewInfo.Location = Data.CameraLocation;
 					CameraInfo.ViewInfo.Rotation = Data.CameraRotation;
 					CameraInfo.ViewInfo.PreviousViewTransform = FTransform(Data.PrevCameraRotation, Data.PrevCameraLocation, FVector::OneVector);
+					CameraInfo.ViewInfo.bConstrainAspectRatio = false;
 				}
 
 				// The Scene View Family must be constructed first as the FSceneView needs it to be constructed
@@ -398,6 +605,26 @@ void FMovieGraphDeferredPanoramicPass::Render(const FMovieGraphTraversalContext&
 
 				CalculateProjectionMatrix(CameraInfo, SceneViewInitOptions, PaneResolution, PaneResolution);
 
+				float MinZ = GNearClippingPlane;
+				const float MaxZ = MinZ;
+				// Avoid zero ViewFOV's which cause divide by zero's in projection matrix
+				const float MatrixFOV = FMath::Max(0.001f, Data.HorizontalFieldOfView) * (float)PI / 360.0f;
+				// ToDo: I think this is a FMath::DegreesToRadians, easier to read that way than PI/360
+				Data.NearClippingPlane = MinZ;
+
+				static_assert((int32)ERHIZBuffer::IsInverted != 0, "ZBuffer should be inverted");
+
+				float XAxisMultiplier = 1.f;
+				float YAxisMultiplier = 1.f;
+				SceneViewInitOptions.ProjectionMatrix = FReversedZPerspectiveMatrix(
+					MatrixFOV,
+					MatrixFOV,
+					XAxisMultiplier,
+					YAxisMultiplier,
+					MinZ,
+					MaxZ
+				);
+
 				// Construct a View to go within this family.
 				FSceneView* NewView = CreateSceneView(SceneViewInitOptions, ViewFamily, CameraInfo);
 
@@ -419,7 +646,7 @@ void FMovieGraphDeferredPanoramicPass::Render(const FMovieGraphTraversalContext&
 				FCanvas Canvas = FCanvas(RenderTargetResource, HitProxyConsumer, GraphRenderer->GetWorld(), GraphRenderer->GetWorld()->GetFeatureLevel(), FCanvas::CDM_DeferDrawing, DPIScale);
 
 				// Construct the sample state that reflects the current render sample
-				UE::MovieGraph::FMovieGraphSampleState SampleState;
+				FMovieGraphPanoSampleState SampleState;
 				{
 					// Take our per-frame Traversal Context and update it with context specific to this sample.
 					FMovieGraphTraversalContext UpdatedTraversalContext = InFrameTraversalContext;
@@ -429,16 +656,16 @@ void FMovieGraphDeferredPanoramicPass::Render(const FMovieGraphTraversalContext&
 					UpdatedTraversalContext.RenderDataIdentifier = RenderDataIdentifier;
 
 					SampleState.TraversalContext = MoveTemp(UpdatedTraversalContext);
-					SampleState.OverscannedResolution = PaneResolution;
+					SampleState.OverscannedResolution = DesiredOutputResolution;
 					SampleState.UnpaddedTileSize = PaneResolution;
 					SampleState.BackbufferResolution = PaneResolution;
 					SampleState.AccumulatorResolution = PaneResolution;
 					SampleState.bWriteSampleToDisk = bWriteAllSamples;
-					SampleState.bRequiresAccumulator = true;
-					SampleState.bFetchFromAccumulator = InTimeData.bIsLastTemporalSampleForFrame && (SpatialIndex == (NumSpatialSamples - 1)) && bIsLastPane;
-					SampleState.OverlappedPad = FIntPoint(0, 0);;
-					SampleState.OverlappedOffset = FIntPoint(0, 0);;
-					SampleState.OverlappedSubpixelShift = FVector2D(0.f, 0.f);
+					SampleState.bRequiresAccumulator = InTimeData.bRequiresAccumulator || (NumSpatialSamples > 1);
+					SampleState.bFetchFromAccumulator = InTimeData.bIsLastTemporalSampleForFrame && (SpatialIndex == (NumSpatialSamples - 1));
+					SampleState.OverlappedPad = FIntPoint(0, 0);
+					SampleState.OverlappedOffset = FIntPoint(0, 0);
+					SampleState.OverlappedSubpixelShift = FVector2D(0.5, 0.5);
 					SampleState.OverscanFraction = 0.f;
 					SampleState.CropRectangle = FIntRect(0, 0, DesiredOutputResolution.X, DesiredOutputResolution.Y); // ToDo: Output resolution will be forced to a 2:1 ratio but this currently respects what the user put in.
 					SampleState.bAllowOCIO = ParentNodeThisFrame->GetAllowOCIO();
@@ -446,6 +673,12 @@ void FMovieGraphDeferredPanoramicPass::Render(const FMovieGraphTraversalContext&
 					SampleState.SceneCaptureSource = SceneCaptureSource;
 					SampleState.CompositingSortOrder = 10;
 					SampleState.RenderLayerIndex = LayerData.LayerIndex;
+					SampleState.Pane = Data;
+
+					// The TileX and TileY are hard-coded to zero to match the MoviePipeline version which did support tiles.
+					SampleState.Debug_OverrideFilename = FString::Printf(TEXT("/%s_SS_%d_TS_%d_TileX_0_TileY_0_PaneX_%d_PaneY_%d.%d"),
+						*SampleState.TraversalContext.RenderDataIdentifier.LayerName, SpatialIndex, InTimeData.TemporalSampleIndex,
+						PaneX, PaneY, InTimeData.OutputFrameNumber);
 				}
 
 				ApplyMovieGraphOverridesToSampleState(SampleState);
@@ -499,12 +732,19 @@ void FMovieGraphDeferredPanoramicPass::Render(const FMovieGraphTraversalContext&
 TSharedRef<::MoviePipeline::IMoviePipelineAccumulationArgs> FMovieGraphDeferredPanoramicPass::GetOrCreateAccumulator(const TObjectPtr<UMovieGraphDefaultRenderer> InGraphRenderer, const FMovieGraphSampleState& InSampleState) const
 {
 	const FMoviePipelineAccumulatorPoolPtr SampleAccumulatorPool = InGraphRenderer->GetOrCreateAccumulatorPool<FImageOverlappedAccumulator>();
-	const DefaultRenderer::FSurfaceAccumulatorPool::FInstancePtr AccumulatorInstance = SampleAccumulatorPool->GetAccumulatorInstance_GameThread<FImageOverlappedAccumulator>(InSampleState.TraversalContext.Time.OutputFrameNumber, InSampleState.TraversalContext.RenderDataIdentifier);
 
+	// Because this is a virtual function and we need to get data out of the polymorphic SampleState 
+	TSharedRef<FMovieGraphPanoSampleState> PanoSampleState = StaticCastSharedRef<FMovieGraphPanoSampleState>(InSampleState.Copy());
+	// Generate a unique PassIdentifier for this Panoramic Pane, to ensure each pane gets its own accumulator.
+	FMovieGraphRenderDataIdentifier RenderDataIdentifierCopy = RenderDataIdentifier;
+	RenderDataIdentifierCopy.SubResourceName = FString::Printf(TEXT("%s_x%d_y%d"), *RenderDataIdentifier.SubResourceName, PanoSampleState->Pane.HorizontalStepIndex, PanoSampleState->Pane.VerticalStepIndex);
+
+
+	const DefaultRenderer::FSurfaceAccumulatorPool::FInstancePtr AccumulatorInstance = SampleAccumulatorPool->GetAccumulatorInstance_GameThread<FImageOverlappedAccumulator>(InSampleState.TraversalContext.Time.OutputFrameNumber, RenderDataIdentifierCopy);
 	TSharedRef<FMovieGraphRenderDataAccumulationArgs> AccumulationArgs = MakeShared<FMovieGraphRenderDataAccumulationArgs>();
 	AccumulationArgs->OutputMerger = PanoramicOutputBlender;
 	AccumulationArgs->ImageAccumulator = StaticCastSharedPtr<FImageOverlappedAccumulator>(AccumulatorInstance->Accumulator);
-	AccumulationArgs->AccumulatorInstance = SampleAccumulatorPool->GetAccumulatorInstance_GameThread<FImageOverlappedAccumulator>(InSampleState.TraversalContext.Time.OutputFrameNumber, InSampleState.TraversalContext.RenderDataIdentifier);
+	AccumulationArgs->AccumulatorInstance = AccumulatorInstance;
 
 	return AccumulationArgs;
 }
diff --git a/Engine/Plugins/MovieScene/MovieRenderPipeline/Source/MovieRenderPipelineRenderPasses/Private/Graph/Renderers/MovieGraphImagePassBase.cpp b/Engine/Plugins/MovieScene/MovieRenderPipeline/Source/MovieRenderPipelineRenderPasses/Private/Graph/Renderers/MovieGraphImagePassBase.cpp
--- a/Engine/Plugins/MovieScene/MovieRenderPipeline/Source/MovieRenderPipelineRenderPasses/Private/Graph/Renderers/MovieGraphImagePassBase.cpp
+++ b/Engine/Plugins/MovieScene/MovieRenderPipeline/Source/MovieRenderPipelineRenderPasses/Private/Graph/Renderers/MovieGraphImagePassBase.cpp
@@ -497,17 +497,20 @@ void FMovieGraphImagePassBase::PostRendererSubmission(
 	TSharedRef<FMovieGraphRenderDataAccumulationArgs> AccumulationArgs =
 		StaticCastSharedRef<FMovieGraphRenderDataAccumulationArgs>(GetOrCreateAccumulator(GraphRenderer, InSampleState));
 
-	auto OnSurfaceReadbackFinished = [this, InSampleState, AccumulationArgs](TUniquePtr<FImagePixelData>&& InPixelData)
+	// Panoramic renders need sample state to be polymorphic, so we don't want to capture the parameter by value.
+	TSharedRef<UE::MovieGraph::FMovieGraphSampleState> SampleStateCopy = InSampleState.Copy();
+
+	auto OnSurfaceReadbackFinished = [this, SampleStateCopy, AccumulationArgs](TUniquePtr<FImagePixelData>&& InPixelData)
 	{
 		FAccumulatorSampleFunc AccumulateFunction = GetAccumulateSampleFunction();
 		
-		const UE::Tasks::TTask<void> Task = UE::Tasks::Launch(UE_SOURCE_LOCATION, [PixelData = MoveTemp(InPixelData), InSampleState, AccumulationArgs, AccumulateFunction]() mutable
+		const UE::Tasks::TTask<void> Task = UE::Tasks::Launch(UE_SOURCE_LOCATION, [PixelData = MoveTemp(InPixelData), SampleStateCopy, AccumulationArgs, AccumulateFunction]() mutable
 		{
 			// Enqueue a encode for this frame onto our worker thread.
-			AccumulateFunction(MoveTemp(PixelData), InSampleState, AccumulationArgs);
+			AccumulateFunction(MoveTemp(PixelData), SampleStateCopy, AccumulationArgs);
 
 			// We have to defer clearing the accumulator until after sample accumulation has finished
-			if (InSampleState.bFetchFromAccumulator)
+			if (SampleStateCopy->bFetchFromAccumulator)
 			{
 				// Final sample has now been executed, free the accumulator for reuse.
 				AccumulationArgs->AccumulatorInstance->SetIsActive(false);
@@ -559,11 +562,12 @@ TFunction<void(TUniquePtr<FImagePixelData>&&)> FMovieGraphImagePassBase::MakeFor
 	// TODO: Is this correct? The prior code is *slightly* different.
 	TSharedRef<FMovieGraphRenderDataAccumulationArgs> AccumulationArgs =
 		StaticCastSharedRef<FMovieGraphRenderDataAccumulationArgs>(GetOrCreateAccumulator(GraphRenderer, InSampleState));
+	TSharedRef<FMovieGraphSampleState> SampleStateCopy = InSampleState.Copy();
 
 	// The legacy surface reader takes the payload just so it can shuffle it into our callback, but we can just include the data
 	// directly in the callback, so this is just a dummy payload.
 	TSharedRef<FImagePixelDataPayload, ESPMode::ThreadSafe> FramePayload = MakeShared<FImagePixelDataPayload, ESPMode::ThreadSafe>();
-	auto Callback = [this, InSampleState, FramePayload, AccumulationArgs](TUniquePtr<FImagePixelData>&& InPixelData)
+	auto Callback = [this, SampleStateCopy, FramePayload, AccumulationArgs](TUniquePtr<FImagePixelData>&& InPixelData)
 	{
 		// Transfer the framePayload to the returned data
 		TUniquePtr<FImagePixelData> PixelDataWithPayload = nullptr;
@@ -593,13 +597,13 @@ TFunction<void(TUniquePtr<FImagePixelData>&&)> FMovieGraphImagePassBase::MakeFor
 
 		FAccumulatorSampleFunc AccumulateFunction = GetAccumulateSampleFunction();
 
-		UE::Tasks::TTask<void> Task = UE::Tasks::Launch(UE_SOURCE_LOCATION, [PixelData = MoveTemp(PixelDataWithPayload), InSampleState, AccumulationArgs, AccumulateFunction]() mutable
+		UE::Tasks::TTask<void> Task = UE::Tasks::Launch(UE_SOURCE_LOCATION, [PixelData = MoveTemp(PixelDataWithPayload), SampleStateCopy, AccumulationArgs, AccumulateFunction]() mutable
 		{
 			// Enqueue a encode for this frame onto our worker thread.
-			AccumulateFunction(MoveTemp(PixelData), InSampleState, AccumulationArgs);
+			AccumulateFunction(MoveTemp(PixelData), SampleStateCopy, AccumulationArgs);
 
 			// We have to defer clearing the accumulator until after sample accumulation has finished
-			if (InSampleState.bFetchFromAccumulator)
+			if (SampleStateCopy->bFetchFromAccumulator)
 			{
 				// Final sample has now been executed, free the accumulator for reuse.
 				AccumulationArgs->AccumulatorInstance->SetIsActive(false);
@@ -631,7 +635,7 @@ FMovieGraphImagePassBase::FAccumulatorSampleFunc FMovieGraphImagePassBase::GetAc
 	return AccumulateSample_TaskThread;
 }
 
-void AccumulateSample_TaskThread(TUniquePtr<FImagePixelData>&& InPixelData, const ::UE::MovieGraph::FMovieGraphSampleState InSampleState, const TSharedRef<::MoviePipeline::IMoviePipelineAccumulationArgs> InAccumulatorArgs)
+void AccumulateSample_TaskThread(TUniquePtr<FImagePixelData>&& InPixelData, const TSharedRef<::UE::MovieGraph::FMovieGraphSampleState> InSampleState, const TSharedRef<::MoviePipeline::IMoviePipelineAccumulationArgs> InAccumulatorArgs)
 {
 	TRACE_CPUPROFILER_EVENT_SCOPE(MoviePipeline_AccumulateSample);
 		
@@ -639,7 +643,7 @@ void AccumulateSample_TaskThread(TUniquePtr<FImagePixelData>&& InPixelData, cons
 
 	// Associate the sample state with the image as payload data, this allows downstream systems to fetch the values without us having to store the data
 	// separately and ensure they stay paired the whole way down.
-	TSharedPtr<FMovieGraphSampleState> SampleStatePayload = MakeShared<FMovieGraphSampleState>(InSampleState);
+	TSharedPtr<FMovieGraphSampleState> SampleStatePayload = InSampleState->Copy();
 	SamplePixelData->SetPayload(StaticCastSharedPtr<IImagePixelDataPayload>(SampleStatePayload));
 
 	const TSharedRef<FMovieGraphRenderDataAccumulationArgs, ESPMode::ThreadSafe> AccumulatorArgs = StaticCastSharedRef<FMovieGraphRenderDataAccumulationArgs>(InAccumulatorArgs);
diff --git a/Engine/Plugins/MovieScene/MovieRenderPipeline/Source/MovieRenderPipelineRenderPasses/Public/Graph/Nodes/MovieGraphDeferredPanoramicPassNode.h b/Engine/Plugins/MovieScene/MovieRenderPipeline/Source/MovieRenderPipelineRenderPasses/Public/Graph/Nodes/MovieGraphDeferredPanoramicPassNode.h
--- a/Engine/Plugins/MovieScene/MovieRenderPipeline/Source/MovieRenderPipelineRenderPasses/Public/Graph/Nodes/MovieGraphDeferredPanoramicPassNode.h
+++ b/Engine/Plugins/MovieScene/MovieRenderPipeline/Source/MovieRenderPipelineRenderPasses/Public/Graph/Nodes/MovieGraphDeferredPanoramicPassNode.h
@@ -3,6 +3,7 @@
 #pragma once
 
 #include "Graph/Nodes/MovieGraphImagePassBaseNode.h"
+#include "MoviePipelinePanoramicBlenderBase.h"
 #include "MovieGraphDeferredPanoramicPassNode.generated.h"
 
 UCLASS(BlueprintType, MinimalAPI)
@@ -11,7 +12,7 @@ class UMovieGraphDeferredPanoramicNode : public UMovieGraphImagePassBaseNode
 	GENERATED_BODY()
 
 public:
-	UMovieGraphDeferredPanoramicNode() = default;
+	UMovieGraphDeferredPanoramicNode();
 
 	// UMovieGraphImagePassBaseNode Interface
 	virtual FEngineShowFlags GetShowFlags() const override;
@@ -58,6 +59,9 @@ public:
 	UPROPERTY(EditAnywhere, BlueprintReadWrite, Category = Overrides, meta = (InlineEditConditionToggle))
 	uint8 bOverride_AntiAliasingMethod : 1;
 
+	UPROPERTY(EditAnywhere, BlueprintReadWrite, Category = Overrides, meta = (InlineEditConditionToggle))
+	uint8 bOverride_Filter : 1;
+
 	UPROPERTY(EditAnywhere, BlueprintReadWrite, Category = Overrides, meta = (InlineEditConditionToggle))
 	uint8 bOverride_bDisableToneCurve : 1;
 
@@ -71,22 +75,30 @@ public:
 	uint8 bOverride_bWriteAllSamples : 1;
 
 	UPROPERTY(EditAnywhere, BlueprintReadWrite, meta = (UIMin = 8, ClampMin = 8), Category = "Sampling", meta = (EditCondition = "bOverride_NumHorizontalSteps"))
-	int32 NumHorizontalSteps = 8;
+	int32 NumHorizontalSteps;
 
 	UPROPERTY(EditAnywhere, BlueprintReadWrite, meta = (UIMin = 3, ClampMin = 3), Category = "Sampling", meta = (EditCondition = "bOverride_NumVerticalSteps"))
-	int32 NumVerticalSteps = 8;
+	int32 NumVerticalSteps;
 	
 	/**
 	* Should the Pitch, Yaw and Roll of the camera be respected? If false, only the location will be taken from the camera.
 	*/
 	UPROPERTY(EditAnywhere, BlueprintReadWrite, Category = "Sampling", meta = (EditCondition = "bOverride_bFollowCameraOrientation"))
-	bool bFollowCameraOrientation = true;
-
+	bool bFollowCameraOrientation;
+	
+	/**
+	* Should we store the render scene history per individual render? This can consume a great deal of memory with many renders,
+	* but enables TAA/TSR and other history-based effects (denoisers, auto-exposure, Lumen, etc.) to work.
+	*/
 	UPROPERTY(EditAnywhere, BlueprintReadWrite, Category = "Sampling", meta = (EditCondition = "bOverride_bAllocateHistoryPerPane"))
-	bool bAllocateHistoryPerPane = false;
-
+	bool bAllocateHistoryPerPane;
+	
+	/*
+	* If true, persistented GPU data per panoramic pane is paged to system memory, allowing higher resolutions, but significantly 
+	* increasing render times. The GPU data is downloaded after each tile is rendered, and then re-uploaded for the next tile.
+	*/
 	UPROPERTY(EditAnywhere, BlueprintReadWrite, Category = "Sampling", meta = (EditCondition = "bOverride_bPageToSystemMemory"))
-	bool bPageToSystemMemory = false;
+	bool bPageToSystemMemory;
 
 	/**
 	* How many sub-pixel jitter renders should we do per temporal sample? This can be used to achieve high
@@ -95,7 +107,7 @@ public:
 	* temporal samples and you will get SpatialSampleCount many renders per temporal sample.
 	*/
 	UPROPERTY(EditAnywhere, BlueprintReadWrite, meta = (UIMin = 1, ClampMin = 1), Category = "Sampling", meta = (EditCondition = "bOverride_SpatialSampleCount"))
-	int32 SpatialSampleCount = 1;
+	int32 SpatialSampleCount;
 
 	/**
 	* Which anti-aliasing method should this render use. If this is set to None, then Movie Render Graph
@@ -107,6 +119,12 @@ public:
 	UPROPERTY(EditAnywhere, BlueprintReadWrite, Category = "Sampling", meta = (EditCondition = "bOverride_AntiAliasingMethod"))
 	TEnumAsByte<EAntiAliasingMethod> AntiAliasingMethod;
 
+	/**
+	* Filter used when blending panoramic. Bilinear is fastest (Samples a 2x2 pixel grid) and produces nearly as good results as the others which require sampling 4x4 pixels.
+	*/
+	UPROPERTY(EditAnywhere, BlueprintReadWrite, Category = "Sampling", meta = (EditCondition = "bOverride_Filter"))
+	EMoviePipelinePanoramicFilterType Filter;
+
 	/**
 	* Debug Feature. Can use this to write out each individual Temporal and Spatial sample rendered by this render pass,
 	* which allows you to see which images are being accumulated together. Can be useful for debugging incorrect looking
diff --git a/Engine/Plugins/MovieScene/MovieRenderPipeline/Source/MovieRenderPipelineRenderPasses/Public/Graph/Renderers/MovieGraphDeferredPanoramicPass.h b/Engine/Plugins/MovieScene/MovieRenderPipeline/Source/MovieRenderPipelineRenderPasses/Public/Graph/Renderers/MovieGraphDeferredPanoramicPass.h
--- a/Engine/Plugins/MovieScene/MovieRenderPipeline/Source/MovieRenderPipelineRenderPasses/Public/Graph/Renderers/MovieGraphDeferredPanoramicPass.h
+++ b/Engine/Plugins/MovieScene/MovieRenderPipeline/Source/MovieRenderPipelineRenderPasses/Public/Graph/Renderers/MovieGraphDeferredPanoramicPass.h
@@ -7,62 +7,6 @@
 
 namespace UE::MovieGraph::Rendering
 {
-
-	class FMovieGraphPanoramicBlender : public UE::MovieGraph::IMovieGraphOutputMerger
-	{
-	public:
-		FMovieGraphPanoramicBlender(TSharedPtr<UE::MovieGraph::IMovieGraphOutputMerger> InOutputMerger, const FIntPoint InOutputResolution)
-		{
-			OutputMerger = InOutputMerger;
-		}
-
-		static FMovieGraphOutputMergerFrame MovieGraphDummyOutputFrame;
-		static TQueue<FMovieGraphOutputMergerFrame> MovieGraphDummyOutputQueue;
-
-		virtual FMovieGraphOutputMergerFrame& AllocateNewOutputFrame_GameThread(const int32 InRenderedFrameNumber) override
-		{
-			// Unsupported, the main Output Builder should be the one tracking this.
-			check(0);
-			return MovieGraphDummyOutputFrame;
-		}
-
-		virtual FMovieGraphOutputMergerFrame& GetOutputFrame_GameThread(const int32 InRenderedFrameNumber) override
-		{
-			// Unsupported, the main Output Builder should be the one tracking this.
-			check(0);
-			return MovieGraphDummyOutputFrame;
-		}
-		virtual void OnCompleteRenderPassDataAvailable_AnyThread(TUniquePtr<FImagePixelData>&& InData)
-		{
-			UE_LOG(LogTemp, Warning, TEXT("OnCompleteRenderPassDataAvailable_AnyThread"));
-		}
-		virtual void OnSingleSampleDataAvailable_AnyThread(TUniquePtr<FImagePixelData>&& InData) override
-		{
-			// This is used for debug output, just pass it straight through.
-			ensure(OutputMerger.IsValid());
-			OutputMerger.Pin()->OnSingleSampleDataAvailable_AnyThread(MoveTemp(InData));
-		}
-		virtual void AbandonOutstandingWork() override
-		{
-			// Not implemented
-			check(0);
-		}
-		virtual int32 GetNumOutstandingFrames() const
-		{
-			return 0; // temp
-		}
-		virtual TQueue<FMovieGraphOutputMergerFrame>& GetFinishedFrames()
-		{
-			// Not implemented
-			check(0);
-			return MovieGraphDummyOutputQueue;
-		}
-
-	private:
-		TWeakPtr<UE::MovieGraph::IMovieGraphOutputMerger> OutputMerger;
-
-	};
-
 	struct MOVIERENDERPIPELINERENDERPASSES_API FMovieGraphDeferredPanoramicPass : public FMovieGraphImagePassBase
 	{
 		FMovieGraphDeferredPanoramicPass();
@@ -98,7 +42,7 @@ namespace UE::MovieGraph::Rendering
 		bool bHasPrintedWarnings = false;
 		FMovieGraphRenderDataIdentifier RenderDataIdentifier;
 		FMovieGraphRenderPassLayerData LayerData;
-		TSharedPtr<FMovieGraphPanoramicBlender> PanoramicOutputBlender;
+		TSharedPtr<UE::MovieGraph::IMovieGraphOutputMerger> PanoramicOutputBlender;
 		
 	};
 }
\ No newline at end of file
diff --git a/Engine/Plugins/MovieScene/MovieRenderPipeline/Source/MovieRenderPipelineRenderPasses/Public/Graph/Renderers/MovieGraphImagePassBase.h b/Engine/Plugins/MovieScene/MovieRenderPipeline/Source/MovieRenderPipelineRenderPasses/Public/Graph/Renderers/MovieGraphImagePassBase.h
--- a/Engine/Plugins/MovieScene/MovieRenderPipeline/Source/MovieRenderPipelineRenderPasses/Public/Graph/Renderers/MovieGraphImagePassBase.h
+++ b/Engine/Plugins/MovieScene/MovieRenderPipeline/Source/MovieRenderPipelineRenderPasses/Public/Graph/Renderers/MovieGraphImagePassBase.h
@@ -34,7 +34,7 @@ namespace UE::MovieGraph::Rendering
 	};
 
 	// Forward Declare
-	void AccumulateSample_TaskThread(TUniquePtr<FImagePixelData>&& InPixelData, const FMovieGraphSampleState InSampleState, const TSharedRef<::MoviePipeline::IMoviePipelineAccumulationArgs> InAccumulatorArgs);
+	void AccumulateSample_TaskThread(TUniquePtr<FImagePixelData>&& InPixelData, const TSharedRef<::UE::MovieGraph::FMovieGraphSampleState> InSampleState, const TSharedRef<::MoviePipeline::IMoviePipelineAccumulationArgs> InAccumulatorArgs);
 
 	struct MOVIERENDERPIPELINERENDERPASSES_API FViewFamilyInitData
 	{
@@ -78,7 +78,7 @@ namespace UE::MovieGraph::Rendering
 		virtual TWeakObjectPtr<UMovieGraphDefaultRenderer> GetRenderer() const { return WeakGraphRenderer; }
 
 	protected:
-		typedef TFunction<void(TUniquePtr<FImagePixelData>&&, const FMovieGraphSampleState, const TSharedRef<::MoviePipeline::IMoviePipelineAccumulationArgs>)> FAccumulatorSampleFunc;
+		typedef TFunction<void(TUniquePtr<FImagePixelData>&&, const TSharedRef<FMovieGraphSampleState>, const TSharedRef<::MoviePipeline::IMoviePipelineAccumulationArgs>)> FAccumulatorSampleFunc;
 		
 		/** Utility function for calculating a Projection Matrix (Orthographic or Perspective), and modify it based on the overscan percentage, aspect ratio, etc. */
 		virtual void CalculateProjectionMatrix(UE::MovieGraph::DefaultRenderer::FCameraInfo& InOutCameraInfo, FSceneViewProjectionData& InOutProjectionData, const FIntPoint InBackbufferResolution, const FIntPoint InAccumulatorResolution) const;
