MovieGraph: Fixed spatial samples not working in ObjectId pass, fixed occasional garbage collection related crash.

diff --git a/Engine/Plugins/MovieScene/MoviePipelineMaskRenderPass/Source/MoviePipelineMaskRenderPass/Private/Graph/MovieGraphObjectIdNode.h b/Engine/Plugins/MovieScene/MoviePipelineMaskRenderPass/Source/MoviePipelineMaskRenderPass/Private/Graph/MovieGraphObjectIdNode.h
--- a/Engine/Plugins/MovieScene/MoviePipelineMaskRenderPass/Source/MoviePipelineMaskRenderPass/Private/Graph/MovieGraphObjectIdNode.h
+++ b/Engine/Plugins/MovieScene/MoviePipelineMaskRenderPass/Source/MoviePipelineMaskRenderPass/Private/Graph/MovieGraphObjectIdNode.h
@@ -25,6 +25,9 @@ public:
 	virtual bool GetAllowsCompositing() const override;
 	virtual int32 GetNumSpatialSamples() const override { return SpatialSampleCount; }
 	virtual int32 GetNumSpatialSamplesDuringWarmUp() const  override { return 0; }
+	// The ObjectID pass doesn't support any anti-aliasing so we force it to None so that their spatial samples turn into MRG-based camera jitter.
+	virtual bool GetOverrideAntiAliasing() const override { return true; }
+	virtual EAntiAliasingMethod GetAntiAliasingMethod() const override { return EAntiAliasingMethod::AAM_None; }
 	// ~UMovieGraphImagePassBaseNode Interface
 
 	// UMovieGraphNode Interface
diff --git a/Engine/Plugins/MovieScene/MoviePipelineMaskRenderPass/Source/MoviePipelineMaskRenderPass/Private/Graph/MovieGraphObjectIdPass.cpp b/Engine/Plugins/MovieScene/MoviePipelineMaskRenderPass/Source/MoviePipelineMaskRenderPass/Private/Graph/MovieGraphObjectIdPass.cpp
--- a/Engine/Plugins/MovieScene/MoviePipelineMaskRenderPass/Source/MoviePipelineMaskRenderPass/Private/Graph/MovieGraphObjectIdPass.cpp
+++ b/Engine/Plugins/MovieScene/MoviePipelineMaskRenderPass/Source/MoviePipelineMaskRenderPass/Private/Graph/MovieGraphObjectIdPass.cpp
@@ -255,8 +255,9 @@ void FMovieGraphObjectIdPass::Render(const FMovieGraphTraversalContext& InFrameT
 {
 	FMovieGraphDeferredPass::Render(InFrameTraversalContext, InTimeData);
 	
-	const UMovieGraphObjectIdNode* ObjectIdNode = Cast<UMovieGraphObjectIdNode>(LayerData.RenderPassNode);
-	check(ObjectIdNode);
+
+	UMovieGraphObjectIdNode* ParentNode = Cast<UMovieGraphObjectIdNode>(GetParentNode(InTimeData.EvaluatedConfig));
+	check(ParentNode);
 	
 	UE::MoviePipeline::FObjectIdAccelerationData* AccelData = GetAccelerationData(LayerData.BranchName);
 	check(AccelData);
@@ -265,7 +266,7 @@ void FMovieGraphObjectIdPass::Render(const FMovieGraphTraversalContext& InFrameT
 	// before the render starts. Render property changes made during the frame (via modifiers) will invalidate the HitProxy IDs in the global array,
 	// and those invalidations happen only once Render() is called (thus we need to wait until after Render() before updating the cache to
 	// get the new IDs).
-	UpdateManifestAccelerationData(*AccelData, ObjectIdNode->IdType);
+	UpdateManifestAccelerationData(*AccelData, ParentNode->IdType);
 }
 
 TSharedRef<MoviePipeline::IMoviePipelineAccumulationArgs> FMovieGraphObjectIdPass::GetOrCreateAccumulator(TObjectPtr<UMovieGraphDefaultRenderer> InGraphRenderer, const UE::MovieGraph::FMovieGraphSampleState& InSampleState) const
