MoviePipeline: Fixed an issue with Object IDs in MRG that could cause some 3rd party applications to not recognize cryptomatte layers properly.

This issue happens when MRG layer name(s) end in a number. For example, if the layer name is Layer1, the cryptomatte layer names that are generated will be Layer100, Layer101, and Layer102. While this is technically not a violation of the cryptomatte specification, it can confuse some 3rd party readers (since they're expecting layer names ending with the exact numbers 00, 01, and 02).

To fix the problem, the cryptomatte layer names that MRG generates will now be of the form "{layer_name}_ObjectID{00-02}" (eg, Layer1_ObjectID00).

diff --git a/Engine/Plugins/MovieScene/MoviePipelineMaskRenderPass/Source/MoviePipelineMaskRenderPass/Private/Graph/MovieGraphObjectIdPass.cpp b/Engine/Plugins/MovieScene/MoviePipelineMaskRenderPass/Source/MoviePipelineMaskRenderPass/Private/Graph/MovieGraphObjectIdPass.cpp
--- a/Engine/Plugins/MovieScene/MoviePipelineMaskRenderPass/Source/MoviePipelineMaskRenderPass/Private/Graph/MovieGraphObjectIdPass.cpp
+++ b/Engine/Plugins/MovieScene/MoviePipelineMaskRenderPass/Source/MoviePipelineMaskRenderPass/Private/Graph/MovieGraphObjectIdPass.cpp
@@ -19,12 +19,12 @@ namespace UE::MovieGraph
 	 */
 	static FString GetCryptomatteTypename(const FMovieGraphRenderDataIdentifier& InRenderDataIdentifier, const bool bIsMultiCam)
 	{
-		FString CryptomatteTypename = FString::Format(TEXT("{0}_{1}"), {InRenderDataIdentifier.RendererName, InRenderDataIdentifier.LayerName});
+		FString CryptomatteTypename = FString::Format(TEXT("{0}_{1}"), {InRenderDataIdentifier.LayerName, InRenderDataIdentifier.RendererName});
 
 		// Include the camera name if there is more than one camera being rendered.
 		if (bIsMultiCam)
 		{
-			CryptomatteTypename = FString::Format(TEXT("{0}_{1}"), {CryptomatteTypename, InRenderDataIdentifier.CameraName});
+			CryptomatteTypename = FString::Format(TEXT("{0}_{1}"), {InRenderDataIdentifier.CameraName, CryptomatteTypename});
 		}
 
 		return CryptomatteTypename;
