MoviePipeline: Fixed a number of minor issues where classes were not correctly exported for usage in C++ based pipelines in licensee projects. Moved some classes from Private folders into Public so that they can be exported and used by other modules, but marked them as final as they're not designed to be extended.

diff --git a/Engine/Plugins/MovieScene/MoviePipelineMaskRenderPass/Source/MoviePipelineMaskRenderPass/Private/Graph/MovieGraphObjectIdNode.cpp b/Engine/Plugins/MovieScene/MoviePipelineMaskRenderPass/Source/MoviePipelineMaskRenderPass/Private/Graph/MovieGraphObjectIdNode.cpp
--- a/Engine/Plugins/MovieScene/MoviePipelineMaskRenderPass/Source/MoviePipelineMaskRenderPass/Private/Graph/MovieGraphObjectIdNode.cpp
+++ b/Engine/Plugins/MovieScene/MoviePipelineMaskRenderPass/Source/MoviePipelineMaskRenderPass/Private/Graph/MovieGraphObjectIdNode.cpp
@@ -1,6 +1,6 @@
-﻿// Copyright Epic Games, Inc. All Rights Reserved.
+// Copyright Epic Games, Inc. All Rights Reserved.
 
-#include "MovieGraphObjectIdNode.h"
+#include "Graph/MovieGraphObjectIdNode.h"
 
 #include "Editor/EditorPerProjectUserSettings.h"
 #include "MovieGraphObjectIdPass.h"
diff --git a/Engine/Plugins/MovieScene/MoviePipelineMaskRenderPass/Source/MoviePipelineMaskRenderPass/Private/Graph/MovieGraphObjectIdNode.h b/Engine/Plugins/MovieScene/MoviePipelineMaskRenderPass/Source/MoviePipelineMaskRenderPass/Private/Graph/MovieGraphObjectIdNode.h
deleted file mode 100644
diff --git a/Engine/Plugins/MovieScene/MoviePipelineMaskRenderPass/Source/MoviePipelineMaskRenderPass/Private/Graph/MovieGraphObjectIdPass.cpp b/Engine/Plugins/MovieScene/MoviePipelineMaskRenderPass/Source/MoviePipelineMaskRenderPass/Private/Graph/MovieGraphObjectIdPass.cpp
--- a/Engine/Plugins/MovieScene/MoviePipelineMaskRenderPass/Source/MoviePipelineMaskRenderPass/Private/Graph/MovieGraphObjectIdPass.cpp
+++ b/Engine/Plugins/MovieScene/MoviePipelineMaskRenderPass/Source/MoviePipelineMaskRenderPass/Private/Graph/MovieGraphObjectIdPass.cpp
@@ -2,7 +2,7 @@
 
 #include "MovieGraphObjectIdPass.h"
 
-#include "MovieGraphObjectIdNode.h"
+#include "Graph/MovieGraphObjectIdNode.h"
 #include "MoviePipelineHashUtils.h"
 #include "MoviePipelineObjectIdUtils.h"
 #include "MoviePipelineQueue.h"
diff --git a/Engine/Plugins/MovieScene/MoviePipelineMaskRenderPass/Source/MoviePipelineMaskRenderPass/Private/MoviePipelineObjectIdPass.h b/Engine/Plugins/MovieScene/MoviePipelineMaskRenderPass/Source/MoviePipelineMaskRenderPass/Private/MoviePipelineObjectIdPass.h
deleted file mode 100644
diff --git a/Engine/Plugins/MovieScene/MoviePipelineMaskRenderPass/Source/MoviePipelineMaskRenderPass/Private/MoviePipelinePanoramicPass.h b/Engine/Plugins/MovieScene/MoviePipelineMaskRenderPass/Source/MoviePipelineMaskRenderPass/Private/MoviePipelinePanoramicPass.h
deleted file mode 100644
diff --git a/Engine/Plugins/MovieScene/MoviePipelineMaskRenderPass/Source/MoviePipelineMaskRenderPass/Public/Graph/MovieGraphObjectIdNode.h b/Engine/Plugins/MovieScene/MoviePipelineMaskRenderPass/Source/MoviePipelineMaskRenderPass/Public/Graph/MovieGraphObjectIdNode.h
new file mode 100644
--- /dev/null
+++ b/Engine/Plugins/MovieScene/MoviePipelineMaskRenderPass/Source/MoviePipelineMaskRenderPass/Public/Graph/MovieGraphObjectIdNode.h
@@ -0,0 +1,84 @@
+// Copyright Epic Games, Inc. All Rights Reserved.
+
+#pragma once
+
+#include "MoviePipelineObjectIdUtils.h"
+#include "Graph/Nodes/MovieGraphImagePassBaseNode.h"
+
+#include "MovieGraphObjectIdNode.generated.h"
+
+enum class EMoviePipelineObjectIdPassIdType : uint8;
+
+/** A node which outputs layers compatible with the Cryptomatte standard. */
+UCLASS(BlueprintType)
+class MOVIEPIPELINEMASKRENDERPASS_API UMovieGraphObjectIdNode final : public UMovieGraphImagePassBaseNode
+{
+	GENERATED_BODY()
+
+public:
+	UMovieGraphObjectIdNode() = default;
+
+	// UMovieGraphImagePassBaseNode Interface
+	virtual FEngineShowFlags GetShowFlags() const override;
+	virtual EViewModeIndex GetViewModeIndex() const override;
+	virtual bool GetAllowsShowFlagsCustomization() const override;
+	virtual bool GetAllowsCompositing() const override;
+	virtual int32 GetNumSpatialSamples() const override { return SpatialSampleCount; }
+	virtual int32 GetNumSpatialSamplesDuringWarmUp() const  override { return 0; }
+	// The ObjectID pass doesn't support any anti-aliasing so we force it to None so that their spatial samples turn into MRG-based camera jitter.
+	virtual bool GetOverrideAntiAliasing() const override { return true; }
+	virtual EAntiAliasingMethod GetAntiAliasingMethod() const override { return EAntiAliasingMethod::AAM_None; }
+	// ~UMovieGraphImagePassBaseNode Interface
+
+	// UMovieGraphNode Interface
+#if WITH_EDITOR
+	virtual FText GetNodeTitle(const bool bGetDescriptive = false) const override;
+	virtual FSlateIcon GetIconAndTint(FLinearColor& OutColor) const override;
+#endif
+	// ~UMovieGraphNode Interface
+
+	virtual void UpdateTelemetry(FMoviePipelineShotRenderTelemetry* InTelemetry) const override;
+
+protected:
+	// UMovieGraphImagePassBaseNode Interface
+	virtual TUniquePtr<UE::MovieGraph::Rendering::FMovieGraphImagePassBase> CreateInstance() const override;
+	// ~UMovieGraphImagePassBaseNode Interface
+
+	// UMovieGraphRenderPassNode Interface
+	virtual FString GetRendererNameImpl() const override;
+	virtual void SetupImpl(const FMovieGraphRenderPassSetupData& InSetupData) override;
+	virtual void TeardownImpl() override;
+	// ~UMovieGraphRenderPassNode Interface
+
+public:
+	UPROPERTY(EditAnywhere, BlueprintReadWrite, Category = Overrides, meta = (InlineEditConditionToggle))
+	uint8 bOverride_SpatialSampleCount : 1;
+
+	UPROPERTY(EditAnywhere, BlueprintReadWrite, Category = Overrides, meta = (InlineEditConditionToggle))
+	uint8 bOverride_IdType : 1;
+
+	UPROPERTY(EditAnywhere, BlueprintReadWrite, Category = Overrides, meta = (InlineEditConditionToggle))
+	uint8 bOverride_bIncludeTranslucentObjects : 1;
+
+	/**
+	* How many sub-pixel jitter renders should we do per temporal sample? This can be used to achieve high
+	* sample counts without Temporal Sub-Sampling (allowing high sample counts without motion blur being enabled).
+	*/
+	UPROPERTY(EditAnywhere, BlueprintReadWrite, meta = (UIMin = 1, ClampMin = 1), Category = "Settings", meta = (EditCondition = "bOverride_SpatialSampleCount"))
+	int32 SpatialSampleCount = 1;
+
+	/** How objects are grouped together within the Cryptomatte. */
+	UPROPERTY(EditAnywhere, BlueprintReadWrite, Category = "Settings", DisplayName = "ID Type", meta = (EditCondition="bOverride_IdType"))
+	EMoviePipelineObjectIdPassIdType IdType = EMoviePipelineObjectIdPassIdType::Full;
+
+	/** If true, translucent objects will be included in the ObjectId pass (but as an opaque layer due to limitations). False will omit translucent objects. */
+	UPROPERTY(EditAnywhere, BlueprintReadWrite, Category = "Settings", meta = (EditCondition="bOverride_bIncludeTranslucentObjects"))
+	bool bIncludeTranslucentObjects = false;
+
+private:
+	/**
+	 * The prior value of the project's bAllowSelectTranslucent setting. This node will temporarily change the value of this property to the value
+	 * set in bIncludeTranslucentObjects while rendering.
+	 */
+	bool bPrevAllowSelectTranslucent;
+};
\ No newline at end of file
diff --git a/Engine/Plugins/MovieScene/MoviePipelineMaskRenderPass/Source/MoviePipelineMaskRenderPass/Public/MoviePipelineObjectIdPass.h b/Engine/Plugins/MovieScene/MoviePipelineMaskRenderPass/Source/MoviePipelineMaskRenderPass/Public/MoviePipelineObjectIdPass.h
new file mode 100644
--- /dev/null
+++ b/Engine/Plugins/MovieScene/MoviePipelineMaskRenderPass/Source/MoviePipelineMaskRenderPass/Public/MoviePipelineObjectIdPass.h
@@ -0,0 +1,49 @@
+// Copyright Epic Games, Inc. All Rights Reserved.
+#pragma once
+
+#include "MoviePipelineImagePassBase.h"
+#include "MoviePipelineObjectIdPass.generated.h"
+
+enum class EMoviePipelineObjectIdPassIdType : uint8;
+
+UCLASS(BlueprintType)
+class MOVIEPIPELINEMASKRENDERPASS_API UMoviePipelineObjectIdRenderPass final : public UMoviePipelineImagePassBase
+{
+	GENERATED_BODY()
+
+public:
+	UMoviePipelineObjectIdRenderPass();
+
+	virtual FText GetDisplayText() const override { return NSLOCTEXT("MovieRenderPipeline", "ObjectIdRenderPassSetting_DisplayName", "Object Ids (Limited)"); }
+	virtual void GetViewShowFlags(FEngineShowFlags& OutShowFlag, EViewModeIndex& OutViewModeIndex) const override;
+	virtual void RenderSample_GameThreadImpl(const FMoviePipelineRenderPassMetrics& InSampleState) override;
+	virtual void TeardownImpl() override;
+	virtual void SetupImpl(const MoviePipeline::FMoviePipelineRenderPassInitSettings& InPassInitSettings) override;
+	virtual void GatherOutputPassesImpl(TArray<FMoviePipelinePassIdentifier>& ExpectedRenderPasses) override;
+	virtual bool IsScreenPercentageSupported() const override { return false; }
+	virtual int32 GetOutputFileSortingOrder() const override { return 10; }
+	virtual void UpdateTelemetry(FMoviePipelineShotRenderTelemetry* InTelemetry) const override;
+	
+	virtual TWeakObjectPtr<UTextureRenderTarget2D> CreateViewRenderTargetImpl(const FIntPoint& InSize, IViewCalcPayload* OptPayload = nullptr) const override;
+	virtual TSharedPtr<FMoviePipelineSurfaceQueue, ESPMode::ThreadSafe> CreateSurfaceQueueImpl(const FIntPoint& InSize, IViewCalcPayload* OptPayload = nullptr) const override;
+
+protected:
+	void PostRendererSubmission(const FMoviePipelineRenderPassMetrics& InSampleState);
+
+private:
+	/** Gets the typename hash used as part of the cryptomatte metadata. Eg, "cryptomatte/<typename_hash>/..." */
+	FString GetTypenameHash() const;
+
+public:
+	UPROPERTY(EditAnywhere, BlueprintReadWrite, Category = "Settings")
+	EMoviePipelineObjectIdPassIdType IdType;
+
+	/** If true, translucent objects will be included in the ObjectId pass (but as an opaque layer due to limitations). False will omit translucent objects. */
+	UPROPERTY(EditAnywhere, BlueprintReadWrite, Category = "Settings")
+	bool bIncludeTranslucentObjects;
+
+private:
+	TSharedPtr<FAccumulatorPool, ESPMode::ThreadSafe> AccumulatorPool;
+	TArray<FMoviePipelinePassIdentifier> ExpectedPassIdentifiers;
+	bool bPrevAllowSelectTranslucent;
+};
diff --git a/Engine/Plugins/MovieScene/MoviePipelineMaskRenderPass/Source/MoviePipelineMaskRenderPass/Public/MoviePipelinePanoramicPass.h b/Engine/Plugins/MovieScene/MoviePipelineMaskRenderPass/Source/MoviePipelineMaskRenderPass/Public/MoviePipelinePanoramicPass.h
new file mode 100644
--- /dev/null
+++ b/Engine/Plugins/MovieScene/MoviePipelineMaskRenderPass/Source/MoviePipelineMaskRenderPass/Public/MoviePipelinePanoramicPass.h
@@ -0,0 +1,184 @@
+// Copyright Epic Games, Inc. All Rights Reserved.
+#pragma once
+
+#include "MoviePipelineImagePassBase.h"
+#include "MoviePipelinePanoramicBlenderBase.h"
+#include "Templates/PimplPtr.h"
+#include "MoviePipelinePanoramicPass.generated.h"
+
+class FCanvas;
+class FOpenColorIODisplayExtension;
+
+class UTextureRenderTarget2D;
+struct FImageOverlappedAccumulator;
+class FSceneViewFamily;
+class FSceneView;
+struct FAccumulatorPool;
+struct FSceneViewStateSystemMemoryMirror;
+
+struct FPanoPane : public UMoviePipelineImagePassBase::IViewCalcPayload
+{
+	// Need to split the shared data (between MoviePipeline and MovieGraph) and want to
+	// avoid multiple inheritance in this struct.
+	UE::MoviePipeline::FPanoramicPane Data;
+};
+
+struct FPanoramicImagePixelDataPayload : public FImagePixelDataPayload
+{
+	virtual TSharedRef<FImagePixelDataPayload> Copy() const override
+	{
+		return MakeShared<FPanoramicImagePixelDataPayload>(*this);
+	}
+		
+	virtual FIntPoint GetAccumulatorSize() const override
+	{
+		return Pane.Data.Resolution;
+	}
+
+	virtual FIntPoint GetOverlapPaddedSize() const override
+	{
+		return Pane.Data.Resolution;
+	}
+
+	virtual bool GetOverlapPaddedSizeIsValid(const FIntPoint InRawSize) const override
+	{
+		// Panoramic images don't support any additional padding/overlap.
+		return InRawSize == Pane.Data.Resolution;
+	}
+
+	virtual void GetWeightFunctionParams(MoviePipeline::FTileWeight1D& WeightFunctionX, MoviePipeline::FTileWeight1D& WeightFunctionY) const override
+	{
+		// 
+		WeightFunctionX.InitHelper(0, Pane.Data.Resolution.X, 0);
+		WeightFunctionY.InitHelper(0, Pane.Data.Resolution.Y, 0);
+	}
+
+	FPanoPane Pane;
+};
+	
+/**
+* Generates a panoramic image (potentially in stereo, stored top/bottom in the final sheet) in equirectangular projection space.
+* Each render is a traditional 2D render and then they are blended together afterwards. For each horizontal step we render n
+* many vertical steps. Each of these renders is called a 'Pane' to avoid confusion with the High Resolution 'Tiles' which only apply to 2D.
+*/
+UCLASS(BlueprintType)
+class MOVIEPIPELINEMASKRENDERPASS_API UMoviePipelinePanoramicPass final : public UMoviePipelineImagePassBase
+{
+	GENERATED_BODY()
+
+public:
+	UMoviePipelinePanoramicPass();
+	
+protected:
+	// UMoviePipelineRenderPass API
+	virtual void SetupImpl(const MoviePipeline::FMoviePipelineRenderPassInitSettings& InPassInitSettings) override;
+	virtual void RenderSample_GameThreadImpl(const FMoviePipelineRenderPassMetrics& InSampleState) override;
+	virtual void TeardownImpl() override;
+#if WITH_EDITOR
+	virtual FText GetDisplayText() const override { return NSLOCTEXT("MovieRenderPipeline", "PanoramicRenderPassSetting_DisplayName", "Panoramic Rendering"); }
+#endif
+	virtual void MoviePipelineRenderShowFlagOverride(FEngineShowFlags& OutShowFlag) override;
+	virtual void GatherOutputPassesImpl(TArray<FMoviePipelinePassIdentifier>& ExpectedRenderPasses) override;
+	virtual bool IsAntiAliasingSupported() const { return true; }
+	virtual int32 GetOutputFileSortingOrder() const override { return 1; }
+	virtual FSceneViewStateInterface* GetSceneViewStateInterface(IViewCalcPayload* OptPayload) override;
+	virtual FSceneViewStateInterface* GetExposureSceneViewStateInterface(IViewCalcPayload* OptPayload, int32 CubeFaceIndex) override;
+	virtual void AddViewExtensions(FSceneViewFamilyContext& InContext, FMoviePipelineRenderPassMetrics& InOutSampleState) override;
+	virtual bool IsAutoExposureAllowed(const FMoviePipelineRenderPassMetrics& InSampleState) const override { return AutoExposureViewStates.Num() > 0; }
+	virtual FSceneView* GetSceneViewForSampleState(FSceneViewFamily* ViewFamily, FMoviePipelineRenderPassMetrics& InOutSampleState, IViewCalcPayload* OptPayload = nullptr) override;
+	virtual TWeakObjectPtr<UTextureRenderTarget2D> GetOrCreateViewRenderTarget(const FIntPoint& InSize, IViewCalcPayload* OptPayload = nullptr) override;
+	virtual TSharedPtr<FMoviePipelineSurfaceQueue, ESPMode::ThreadSafe> GetOrCreateSurfaceQueue(const FIntPoint& InSize, IViewCalcPayload* OptPayload = nullptr) override;
+	// ~UMoviePipelineRenderPass
+
+	// FGCObject Interface
+	static void AddReferencedObjects(UObject* InThis, FReferenceCollector& Collector);
+	// ~FGCObject Interface
+
+	void ScheduleReadbackAndAccumulation(const FMoviePipelineRenderPassMetrics& InSampleState, const FPanoPane& InPane, FCanvas& InCanvas);
+	void GetFieldOfView(float& OutHorizontal, float& OutVertical, const bool bInStereo) const;
+	FIntPoint GetPaneResolution(const FIntPoint& InSize) const;
+	FIntPoint GetPayloadPaneResolution(const FIntPoint& InSize, IViewCalcPayload* OptPayload) const;
+
+	virtual void UpdateTelemetry(FMoviePipelineShotRenderTelemetry* InTelemetry) const override;
+
+public:
+
+	/**
+	* How many different renders should the 360* horizontal view be broken into? Higher numbers are less distorted but longer to render.
+	*/
+	UPROPERTY(EditAnywhere, BlueprintReadWrite, Category = "Panoramic Settings", meta = (UIMin = "1", ClampMin = "1"))
+	int32 NumHorizontalSteps;
+	
+	/**
+	* How many different renders should the 360* vertical view be broken into? Higher numbers are less distorted but longer to render.
+	* This is typically less than horizontal as the vertical poles of an image will always have distortion in equirectangular projections.
+	*/
+	UPROPERTY(EditAnywhere, BlueprintReadWrite, Category = "Panoramic Settings", meta = (UIMin = "1", ClampMin = "1"))
+	int32 NumVerticalSteps;
+
+	/**
+	* Should the Pitch, Yaw and Roll of the camera be respected? If false, only the location will be taken from the camera.
+	*/
+	UPROPERTY(EditAnywhere, BlueprintReadWrite, Category = "Panoramic Settings")
+	bool bFollowCameraOrientation;
+	
+	/**
+	* Filter used when blending panoramic. Bilinear is fastest (Samples a 2x2 pixel grid) and produces nearly as good results as the others which require sampling 4x4 pixels.
+	*/
+	UPROPERTY(EditAnywhere, BlueprintReadWrite, Category = "Panoramic Settings")
+	EMoviePipelinePanoramicFilterType Filter = EMoviePipelinePanoramicFilterType::Bilinear;
+	
+	/**
+	* If true we will capture a stereo panorama. This doubles everything (all render times, memory requirements, etc.)
+	* May need a significantly higher number of horizontal steps to look good when viewed in stereo.
+	*/
+	// UPROPERTY(EditAnywhere, BlueprintReadWrite, Category = "Panoramic Settings")
+	bool bStereo;
+
+	//UPROPERTY(EditAnywhere, BlueprintReadWrite, AdvancedDisplay, Category = "Panoramic Settings")
+	float EyeSeparation;
+	
+	//UPROPERTY(EditAnywhere, BlueprintReadWrite, AdvancedDisplay, Category = "Panoramic Settings")
+	float EyeConvergenceDistance;
+
+	/**Advance used only. Allows you to override the Horizontal Field of View (if not zero). Can cause crashes or incomplete panoramas.*/
+	UPROPERTY(EditAnywhere, BlueprintReadWrite, AdvancedDisplay, Category = "Panoramic Settings")
+	float HorzFieldOfView = 0.f;
+
+	/**Advance used only. Allows you to override the Vertical Field of View (if not zero). Can cause crashes or incomplete panoramas.*/
+	UPROPERTY(EditAnywhere, BlueprintReadWrite, AdvancedDisplay, Category = "Panoramic Settings")
+	float VertFieldOfView = 0.f;
+
+	// UPROPERTY(EditAnywhere, BlueprintReadWrite, Category = "Panoramic Settings")
+	bool bAccumulatorIncludesAlpha;
+	
+	/**
+	* Should we store the render scene history per individual render? This can consume a great deal of memory with many renders,
+	* but enables TAA and other history-based effects (denoisers, auto-exposure, Lumen, etc.) to work.
+	*/
+	UPROPERTY(EditAnywhere, BlueprintReadWrite, AdvancedDisplay, Category = "Panoramic Settings")
+	bool bAllocateHistoryPerPane;
+
+	/*
+	* If true, persistent GPU data per panoramic pane is mirrored to system memory, allowing higher resolutions, but increasing render time.
+	*/
+	UPROPERTY(EditAnywhere, BlueprintReadWrite, Category = "Render Settings")
+	bool bPageToSystemMemory;
+
+protected:
+	TSharedPtr<FAccumulatorPool, ESPMode::ThreadSafe> AccumulatorPool;
+
+	// ToDo: One per high-res tile per pano-pane?
+	TArray<FSceneViewStateReference> OptionalPaneViewStates;
+
+	// When using an auto exposure render pass, holds view states for 6 cube faces
+	TArray<FSceneViewStateReference> AutoExposureViewStates;
+
+	/** The lifetime of this SceneViewExtension is only during the rendering process. It is destroyed as part of TearDown. */
+	TSharedPtr<FOpenColorIODisplayExtension, ESPMode::ThreadSafe> OCIOSceneViewExtension;
+
+	TSharedPtr<MoviePipeline::IMoviePipelineOutputMerger> PanoramicOutputBlender;
+	bool bHasWarnedSettings;
+
+	TPimplPtr<FSceneViewStateSystemMemoryMirror> SystemMemoryMirror;
+};
diff --git a/Engine/Plugins/MovieScene/MovieRenderPipeline/Source/MovieRenderPipelineCore/Private/Graph/Nodes/MovieGraphWidgetRendererBaseNode.h b/Engine/Plugins/MovieScene/MovieRenderPipeline/Source/MovieRenderPipelineCore/Private/Graph/Nodes/MovieGraphWidgetRendererBaseNode.h
deleted file mode 100644
diff --git a/Engine/Plugins/MovieScene/MovieRenderPipeline/Source/MovieRenderPipelineCore/Public/Graph/MovieGraphConfig.h b/Engine/Plugins/MovieScene/MovieRenderPipeline/Source/MovieRenderPipelineCore/Public/Graph/MovieGraphConfig.h
--- a/Engine/Plugins/MovieScene/MovieRenderPipeline/Source/MovieRenderPipelineCore/Public/Graph/MovieGraphConfig.h
+++ b/Engine/Plugins/MovieScene/MovieRenderPipeline/Source/MovieRenderPipelineCore/Public/Graph/MovieGraphConfig.h
@@ -568,7 +568,7 @@ class UMovieGraphTraversableObject : public UInterface
  * Provides a way for objects, which would otherwise not be mergeable during a traversal, to merge in a well-defined way.
  * Also allows objects to expose which properties have been affected by the merge.
  */
-class IMovieGraphTraversableObject
+class MOVIERENDERPIPELINECORE_API IMovieGraphTraversableObject
 {
 	GENERATED_BODY()
 
diff --git a/Engine/Plugins/MovieScene/MovieRenderPipeline/Source/MovieRenderPipelineCore/Public/Graph/MovieGraphRenderLayerSubsystem.h b/Engine/Plugins/MovieScene/MovieRenderPipeline/Source/MovieRenderPipelineCore/Public/Graph/MovieGraphRenderLayerSubsystem.h
--- a/Engine/Plugins/MovieScene/MovieRenderPipeline/Source/MovieRenderPipelineCore/Public/Graph/MovieGraphRenderLayerSubsystem.h
+++ b/Engine/Plugins/MovieScene/MovieRenderPipeline/Source/MovieRenderPipelineCore/Public/Graph/MovieGraphRenderLayerSubsystem.h
@@ -1062,7 +1062,7 @@ public:
  * Provides a means of assembling modifiers together to generate a desired view of a scene. 
  */
 UCLASS(BlueprintType)
-class UMovieGraphRenderLayer : public UObject
+class MOVIERENDERPIPELINECORE_API UMovieGraphRenderLayer : public UObject
 {
 	GENERATED_BODY()
 
diff --git a/Engine/Plugins/MovieScene/MovieRenderPipeline/Source/MovieRenderPipelineCore/Public/Graph/Nodes/MovieGraphAudioOutputNode.h b/Engine/Plugins/MovieScene/MovieRenderPipeline/Source/MovieRenderPipelineCore/Public/Graph/Nodes/MovieGraphAudioOutputNode.h
--- a/Engine/Plugins/MovieScene/MovieRenderPipeline/Source/MovieRenderPipelineCore/Public/Graph/Nodes/MovieGraphAudioOutputNode.h
+++ b/Engine/Plugins/MovieScene/MovieRenderPipeline/Source/MovieRenderPipelineCore/Public/Graph/Nodes/MovieGraphAudioOutputNode.h
@@ -11,7 +11,7 @@
  * A node which outputs .wav audio after all renders have completed.
  */
 UCLASS()
-class UMovieGraphAudioOutputNode : public UMovieGraphFileOutputNode
+class MOVIERENDERPIPELINECORE_API UMovieGraphAudioOutputNode : public UMovieGraphFileOutputNode
 {
 	GENERATED_BODY()
 	
diff --git a/Engine/Plugins/MovieScene/MovieRenderPipeline/Source/MovieRenderPipelineCore/Public/Graph/Nodes/MovieGraphWidgetRendererBaseNode.h b/Engine/Plugins/MovieScene/MovieRenderPipeline/Source/MovieRenderPipelineCore/Public/Graph/Nodes/MovieGraphWidgetRendererBaseNode.h
new file mode 100644
--- /dev/null
+++ b/Engine/Plugins/MovieScene/MovieRenderPipeline/Source/MovieRenderPipelineCore/Public/Graph/Nodes/MovieGraphWidgetRendererBaseNode.h
@@ -0,0 +1,82 @@
+// Copyright Epic Games, Inc. All Rights Reserved.
+
+#pragma once
+
+#include "Graph/Nodes/MovieGraphRenderPassNode.h"
+
+#include "Styling/AppStyle.h"
+#include "MovieGraphWidgetRendererBaseNode.generated.h"
+
+// Forward Declares
+class FWidgetRenderer;
+class SVirtualWindow;
+class SWidget;
+class UMovieGraphDefaultRenderer;
+class UTextureRenderTarget2D;
+struct FMovieGraphRenderPassLayerData;
+struct FMovieGraphRenderPassSetupData;
+
+/** Base node containing common logic for nodes that render widgets. */
+UCLASS(Abstract)
+class MOVIERENDERPIPELINECORE_API UMovieGraphWidgetRendererBaseNode : public UMovieGraphRenderPassNode
+{
+	GENERATED_BODY()
+
+public:
+	UPROPERTY(EditAnywhere, BlueprintReadWrite, Category = Overrides, meta = (InlineEditConditionToggle))
+	uint8 bOverride_bCompositeOntoFinalImage : 1;
+
+	/** If true, the pass will be composited onto each render. Does not apply to multi-layer EXR files. */
+	UPROPERTY(EditAnywhere, BlueprintReadWrite, meta=(EditCondition="bOverride_bCompositeOntoFinalImage"), Category="Widget Settings")
+	bool bCompositeOntoFinalImage = true;
+
+protected:
+	/** A render pass for a specific render layer, generated by the UMovieGraphWidgetRendererBaseNode CDO. */
+	struct FMovieGraphWidgetPass
+	{
+		virtual ~FMovieGraphWidgetPass() = default;
+
+		virtual void Setup(TWeakObjectPtr<UMovieGraphDefaultRenderer> InRenderer, const FMovieGraphRenderPassLayerData& InLayer);
+		virtual void Teardown();
+		virtual TSharedPtr<SWidget> GetWidget(UMovieGraphWidgetRendererBaseNode* InNodeThisFrame) = 0;
+		virtual void Render(const FMovieGraphTraversalContext& InFrameTraversalContext, const FMovieGraphTimeStepData& InTimeData);
+		virtual void GatherOutputPasses(TArray<FMovieGraphRenderDataIdentifier>& OutExpectedPasses) const;
+		virtual int32 GetCompositingSortOrder() const = 0;
+		virtual UMovieGraphWidgetRendererBaseNode* GetParentNode(UMovieGraphEvaluatedConfig* InConfig) const;
+
+
+	protected:
+		FMovieGraphRenderPassLayerData LayerData;
+
+		/** Unique identifier passed in GatherOutputPasses and with each render that identifies the data produced by this renderer. */
+		FMovieGraphRenderDataIdentifier RenderDataIdentifier;
+
+		/** The renderer associated with this pass. */
+		TWeakObjectPtr<UMovieGraphDefaultRenderer> Renderer;
+	};
+
+	/** Widget renderer shared with all FMovieGraphWidgetPass instances. */
+	TSharedPtr<FWidgetRenderer> WidgetRenderer;
+
+	/** Render passes generated by GeneratePass() calls on the CDO. */
+	TArray<TUniquePtr<FMovieGraphWidgetPass>> CurrentInstances;
+
+protected:
+	/** Generate a new pass that the widget renderer will use. */
+	virtual TUniquePtr<FMovieGraphWidgetPass> GeneratePass() PURE_VIRTUAL(UMovieGraphWidgetRendererBaseNode::GeneratePass, return nullptr;);
+	
+	// UMovieGraphRenderPassNode Interface
+	virtual void SetupImpl(const FMovieGraphRenderPassSetupData& InSetupData) override;
+	virtual void TeardownImpl() override;
+	virtual void RenderImpl(const FMovieGraphTraversalContext& InFrameTraversalContext, const FMovieGraphTimeStepData& InTimeData) override;
+	virtual void GatherOutputPassesImpl(UMovieGraphEvaluatedConfig* InConfig, TArray<FMovieGraphRenderDataIdentifier>& OutExpectedPasses) const override;
+	// ~UMovieGraphRenderPassNode Interface
+
+private:
+	/** Virtual window instances shared with all FMovieGraphWidgetPass instances, keyed by resolution. */
+	TMap<const FIntPoint, TSharedPtr<SVirtualWindow>> SharedVirtualWindows;
+
+private:
+	/** Gets an existing virtual window with the specified resolution, or creates a new one if one does not exist yet. */
+	TSharedPtr<SVirtualWindow> GetOrCreateVirtualWindow(const FIntPoint& InResolution);
+};
\ No newline at end of file
diff --git a/Engine/Plugins/MovieScene/MovieRenderPipeline/Source/MovieRenderPipelineMP4Encoder/Public/Graph/MovieGraphMP4EncoderNode.h b/Engine/Plugins/MovieScene/MovieRenderPipeline/Source/MovieRenderPipelineMP4Encoder/Public/Graph/MovieGraphMP4EncoderNode.h
--- a/Engine/Plugins/MovieScene/MovieRenderPipeline/Source/MovieRenderPipelineMP4Encoder/Public/Graph/MovieGraphMP4EncoderNode.h
+++ b/Engine/Plugins/MovieScene/MovieRenderPipeline/Source/MovieRenderPipelineMP4Encoder/Public/Graph/MovieGraphMP4EncoderNode.h
@@ -10,7 +10,7 @@
 
 /** A node which can output H264 mp4 files. */
 UCLASS(BlueprintType)
-class UMovieGraphMP4EncoderNode : public UMovieGraphVideoOutputNode
+class MOVIERENDERPIPELINEMP4ENCODER_API UMovieGraphMP4EncoderNode final : public UMovieGraphVideoOutputNode
 {
 	GENERATED_BODY()
 
diff --git a/Engine/Plugins/MovieScene/MovieRenderPipeline/Source/MovieRenderPipelineMP4Encoder/Public/MoviePipelineMP4EncoderOutput.h b/Engine/Plugins/MovieScene/MovieRenderPipeline/Source/MovieRenderPipelineMP4Encoder/Public/MoviePipelineMP4EncoderOutput.h
--- a/Engine/Plugins/MovieScene/MovieRenderPipeline/Source/MovieRenderPipelineMP4Encoder/Public/MoviePipelineMP4EncoderOutput.h
+++ b/Engine/Plugins/MovieScene/MovieRenderPipeline/Source/MovieRenderPipelineMP4Encoder/Public/MoviePipelineMP4EncoderOutput.h
@@ -9,7 +9,7 @@
 class FMoviePipelineMP4Encoder;
 
 UCLASS(BlueprintType)
-class UMoviePipelineMP4EncoderOutput : public UMoviePipelineVideoOutputBase
+class MOVIERENDERPIPELINEMP4ENCODER_API UMoviePipelineMP4EncoderOutput final : public UMoviePipelineVideoOutputBase
 {
 	GENERATED_BODY()
 
diff --git a/Engine/Plugins/MovieScene/MovieRenderPipeline/Source/MovieRenderPipelineRenderPasses/Public/Graph/Nodes/MovieGraphDeferredPanoramicPassNode.h b/Engine/Plugins/MovieScene/MovieRenderPipeline/Source/MovieRenderPipelineRenderPasses/Public/Graph/Nodes/MovieGraphDeferredPanoramicPassNode.h
--- a/Engine/Plugins/MovieScene/MovieRenderPipeline/Source/MovieRenderPipelineRenderPasses/Public/Graph/Nodes/MovieGraphDeferredPanoramicPassNode.h
+++ b/Engine/Plugins/MovieScene/MovieRenderPipeline/Source/MovieRenderPipelineRenderPasses/Public/Graph/Nodes/MovieGraphDeferredPanoramicPassNode.h
@@ -6,8 +6,8 @@
 #include "MoviePipelinePanoramicBlenderBase.h"
 #include "MovieGraphDeferredPanoramicPassNode.generated.h"
 
-UCLASS(BlueprintType, MinimalAPI)
-class UMovieGraphDeferredPanoramicNode : public UMovieGraphImagePassBaseNode
+UCLASS(BlueprintType)
+class MOVIERENDERPIPELINERENDERPASSES_API UMovieGraphDeferredPanoramicNode final : public UMovieGraphImagePassBaseNode
 {
 	GENERATED_BODY()
 
diff --git a/Engine/Plugins/MovieScene/MovieRenderPipeline/Source/MovieRenderPipelineRenderPasses/Public/MovieGraphImageSequenceOutputNode.h b/Engine/Plugins/MovieScene/MovieRenderPipeline/Source/MovieRenderPipelineRenderPasses/Public/MovieGraphImageSequenceOutputNode.h
--- a/Engine/Plugins/MovieScene/MovieRenderPipeline/Source/MovieRenderPipelineRenderPasses/Public/MovieGraphImageSequenceOutputNode.h
+++ b/Engine/Plugins/MovieScene/MovieRenderPipeline/Source/MovieRenderPipelineRenderPasses/Public/MovieGraphImageSequenceOutputNode.h
@@ -93,7 +93,7 @@ protected:
  * Image sequence output node that can write single-layer EXR files.
  */
 UCLASS()
-class UMovieGraphImageSequenceOutputNode_EXR : public UMovieGraphImageSequenceOutputNode
+class MOVIERENDERPIPELINERENDERPASSES_API UMovieGraphImageSequenceOutputNode_EXR : public UMovieGraphImageSequenceOutputNode
 {
 	GENERATED_BODY()
 
@@ -177,7 +177,7 @@ protected:
  * Image sequence output node that can write multi-layer EXR files.
  */
 UCLASS()
-class UMovieGraphImageSequenceOutputNode_MultiLayerEXR : public UMovieGraphImageSequenceOutputNode_EXR
+class MOVIERENDERPIPELINERENDERPASSES_API UMovieGraphImageSequenceOutputNode_MultiLayerEXR : public UMovieGraphImageSequenceOutputNode_EXR
 {
 	GENERATED_BODY()
 
@@ -280,7 +280,7 @@ public:
 * No metadata is supported.
 */
 UCLASS()
-class UMovieGraphImageSequenceOutputNode_BMP : public UMovieGraphImageSequenceOutputNode
+class MOVIERENDERPIPELINERENDERPASSES_API UMovieGraphImageSequenceOutputNode_BMP : public UMovieGraphImageSequenceOutputNode
 {
 	GENERATED_BODY()
 
@@ -329,7 +329,7 @@ public:
 * No metadata is supported.
 */
 UCLASS()
-class UMovieGraphImageSequenceOutputNode_JPG : public UMovieGraphImageSequenceOutputNode
+class MOVIERENDERPIPELINERENDERPASSES_API UMovieGraphImageSequenceOutputNode_JPG : public UMovieGraphImageSequenceOutputNode
 {
 	GENERATED_BODY()
 
@@ -377,7 +377,7 @@ public:
 * No metadata is supported.
 */
 UCLASS()
-class UMovieGraphImageSequenceOutputNode_PNG : public UMovieGraphImageSequenceOutputNode
+class MOVIERENDERPIPELINERENDERPASSES_API UMovieGraphImageSequenceOutputNode_PNG : public UMovieGraphImageSequenceOutputNode
 {
 	GENERATED_BODY()
 
