MoviePipeline: Fixed an issue in MRG that could cause Object IDs to flicker when they are used across multiple layers.

diff --git a/Engine/Plugins/MovieScene/MoviePipelineMaskRenderPass/Source/MoviePipelineMaskRenderPass/Private/Graph/MovieGraphObjectIdPass.cpp b/Engine/Plugins/MovieScene/MoviePipelineMaskRenderPass/Source/MoviePipelineMaskRenderPass/Private/Graph/MovieGraphObjectIdPass.cpp
--- a/Engine/Plugins/MovieScene/MoviePipelineMaskRenderPass/Source/MoviePipelineMaskRenderPass/Private/Graph/MovieGraphObjectIdPass.cpp
+++ b/Engine/Plugins/MovieScene/MoviePipelineMaskRenderPass/Source/MoviePipelineMaskRenderPass/Private/Graph/MovieGraphObjectIdPass.cpp
@@ -219,11 +219,6 @@ void FMovieGraphObjectIdPass::Setup(TWeakObjectPtr<UMovieGraphDefaultRenderer> I
 		// Add our default to the manifest
 		static const uint32 DefaultHash = ::MoviePipeline::HashNameToId(TCHAR_TO_UTF8(TEXT("default")));
 		AccelData.JsonManifest->SetStringField(TEXT("default"), FString::Printf(TEXT("%08x"), DefaultHash));
-
-		// Add the HitProxies
-		const UMovieGraphObjectIdNode* ObjectIdNode = Cast<UMovieGraphObjectIdNode>(LayerData.RenderPassNode);
-		check(ObjectIdNode);
-		UpdateManifestAccelerationData(AccelData, ObjectIdNode->IdType);
 	}
 }
 
@@ -254,29 +249,29 @@ UMovieGraphImagePassBaseNode* FMovieGraphObjectIdPass::GetParentNode(UMovieGraph
 void FMovieGraphObjectIdPass::Render(const FMovieGraphTraversalContext& InFrameTraversalContext, const FMovieGraphTimeStepData& InTimeData)
 {
 	FMovieGraphDeferredPass::Render(InFrameTraversalContext, InTimeData);
-	
-
-	UMovieGraphObjectIdNode* ParentNode = Cast<UMovieGraphObjectIdNode>(GetParentNode(InTimeData.EvaluatedConfig));
-	check(ParentNode);
-	
-	UE::MoviePipeline::FObjectIdAccelerationData* AccelData = GetAccelerationData(LayerData.BranchName);
-	check(AccelData);
-	
-	// This needs to be updated every frame. It's also important to call this AFTER the render finishes; HitProxy IDs will be out-of-date
-	// before the render starts. Render property changes made during the frame (via modifiers) will invalidate the HitProxy IDs in the global array,
-	// and those invalidations happen only once Render() is called (thus we need to wait until after Render() before updating the cache to
-	// get the new IDs).
-	UpdateManifestAccelerationData(*AccelData, ParentNode->IdType);
 }
 
 TSharedRef<MoviePipeline::IMoviePipelineAccumulationArgs> FMovieGraphObjectIdPass::GetOrCreateAccumulator(TObjectPtr<UMovieGraphDefaultRenderer> InGraphRenderer, const UE::MovieGraph::FMovieGraphSampleState& InSampleState) const
 {
+	UE::MoviePipeline::FObjectIdAccelerationData* AccelerationData = GetAccelerationData(InSampleState.TraversalContext.RenderDataIdentifier.RootBranchName);
+	check(AccelerationData);
+
+	// Update the HitProxy acceleration cache; this needs to be updated every frame. It's important to call this AFTER the render finishes because
+	// HitProxy IDs will be out-of-date before the render starts. Render property changes made during the frame (via modifiers) will invalidate the
+	// HitProxy IDs in the global array, and those invalidations happen only once Render() is called.
+	//
+	// Doing the update here is the best spot as far as timing goes. At this point, the render has finished, modifiers are still applied, and the next
+	// layer's render has not started yet, which ensures that the global HitProxy array is still reflective of what was just rendered.
+	{
+		const UMovieGraphObjectIdNode* ParentNode = Cast<UMovieGraphObjectIdNode>(GetParentNode(InSampleState.TraversalContext.Time.EvaluatedConfig));
+		check(ParentNode);
+
+		UpdateManifestAccelerationData(*AccelerationData, ParentNode->IdType);
+	}
+
 	const FMoviePipelineAccumulatorPoolPtr SampleAccumulatorPool = InGraphRenderer->GetOrCreateAccumulatorPool<FMaskOverlappedAccumulator>();
 	const UE::MovieGraph::DefaultRenderer::FSurfaceAccumulatorPool::FInstancePtr AccumulatorInstance = SampleAccumulatorPool->GetAccumulatorInstance_GameThread<FMaskOverlappedAccumulator>(InSampleState.TraversalContext.Time.OutputFrameNumber, InSampleState.TraversalContext.RenderDataIdentifier);
 
-	const UE::MoviePipeline::FObjectIdAccelerationData* AccelerationData = GetAccelerationData(InSampleState.TraversalContext.RenderDataIdentifier.RootBranchName);
-	check(AccelerationData);
-
 	TSharedRef<FMovieGraphObjectIdMaskSampleAccumulationArgs> AccumulationArgs = MakeShared<FMovieGraphObjectIdMaskSampleAccumulationArgs>();
 	AccumulationArgs->OutputMerger = InGraphRenderer->GetOwningGraph()->GetOutputMerger();
 	AccumulationArgs->ImageAccumulator = StaticCastSharedPtr<FMaskOverlappedAccumulator>(AccumulatorInstance->Accumulator);
