MoviePipeline: Clamp the custom FoV overrides to [1, 179] to prevent users from trying to create >= 180* views which cause incorrect calculations that results in a crash.

diff --git a/Engine/Plugins/MovieScene/MoviePipelineMaskRenderPass/Source/MoviePipelineMaskRenderPass/Private/MoviePipelinePanoramicPass.cpp b/Engine/Plugins/MovieScene/MoviePipelineMaskRenderPass/Source/MoviePipelineMaskRenderPass/Private/MoviePipelinePanoramicPass.cpp
--- a/Engine/Plugins/MovieScene/MoviePipelineMaskRenderPass/Source/MoviePipelineMaskRenderPass/Private/MoviePipelinePanoramicPass.cpp
+++ b/Engine/Plugins/MovieScene/MoviePipelineMaskRenderPass/Source/MoviePipelineMaskRenderPass/Private/MoviePipelinePanoramicPass.cpp
@@ -248,6 +248,10 @@ FIntPoint UMoviePipelinePanoramicPass::GetPaneResolution(const FIntPoint& InSize
 	float VerticalFoV;
 	GetFieldOfView(HorizontalFoV, VerticalFoV, bStereo);
 
+	// Clamp them to a reasonable field of view to avoid spherical pole related singularities.
+	HorizontalFoV = FMath::Clamp(HorizontalFoV, 1, 179);
+	VerticalFoV = FMath::Clamp(VerticalFoV, 1, 179);
+
 	// Horizontal FoV is a proportion of the global horizontal resolution
 	// ToDo: We might have to check which is higher, if numVerticalPanes > numHorizontalPanes this math might be backwards.
 	float HorizontalRes = (HorizontalFoV / 360.0f) * InSize.X;
diff --git a/Engine/Plugins/MovieScene/MoviePipelineMaskRenderPass/Source/MoviePipelineMaskRenderPass/Public/MoviePipelinePanoramicPass.h b/Engine/Plugins/MovieScene/MoviePipelineMaskRenderPass/Source/MoviePipelineMaskRenderPass/Public/MoviePipelinePanoramicPass.h
--- a/Engine/Plugins/MovieScene/MoviePipelineMaskRenderPass/Source/MoviePipelineMaskRenderPass/Public/MoviePipelinePanoramicPass.h
+++ b/Engine/Plugins/MovieScene/MoviePipelineMaskRenderPass/Source/MoviePipelineMaskRenderPass/Public/MoviePipelinePanoramicPass.h
@@ -144,11 +144,11 @@ public:
 	float EyeConvergenceDistance;
 
 	/**Advance used only. Allows you to override the Horizontal Field of View (if not zero). Can cause crashes or incomplete panoramas.*/
-	UPROPERTY(EditAnywhere, BlueprintReadWrite, AdvancedDisplay, Category = "Panoramic Settings")
+	UPROPERTY(EditAnywhere, BlueprintReadWrite, AdvancedDisplay, Category = "Panoramic Settings", meta = (UIMin = "1", ClampMin = "1", UIMax="179", ClampMax = "179"))
 	float HorzFieldOfView = 0.f;
 
 	/**Advance used only. Allows you to override the Vertical Field of View (if not zero). Can cause crashes or incomplete panoramas.*/
-	UPROPERTY(EditAnywhere, BlueprintReadWrite, AdvancedDisplay, Category = "Panoramic Settings")
+	UPROPERTY(EditAnywhere, BlueprintReadWrite, AdvancedDisplay, Category = "Panoramic Settings", meta = (UIMin = "1", ClampMin = "1", UIMax = "179", ClampMax = "179"))
 	float VertFieldOfView = 0.f;
 
 	// UPROPERTY(EditAnywhere, BlueprintReadWrite, Category = "Panoramic Settings")
