MoviePipeline: Fixed a bug in the Object ID implementation that caused the "Include Translucent Objects" property to not take effect. Also enabled this property to vary between layers.

diff --git a/Engine/Plugins/MovieScene/MoviePipelineMaskRenderPass/Source/MoviePipelineMaskRenderPass/Private/Graph/MovieGraphObjectIdNode.cpp b/Engine/Plugins/MovieScene/MoviePipelineMaskRenderPass/Source/MoviePipelineMaskRenderPass/Private/Graph/MovieGraphObjectIdNode.cpp
--- a/Engine/Plugins/MovieScene/MoviePipelineMaskRenderPass/Source/MoviePipelineMaskRenderPass/Private/Graph/MovieGraphObjectIdNode.cpp
+++ b/Engine/Plugins/MovieScene/MoviePipelineMaskRenderPass/Source/MoviePipelineMaskRenderPass/Private/Graph/MovieGraphObjectIdNode.cpp
@@ -100,24 +100,3 @@ TUniquePtr<UE::MovieGraph::Rendering::FMovieGraphImagePassBase> UMovieGraphObjec
 {
 	return MakeUnique<FMovieGraphObjectIdPass>();
 }
-
-void UMovieGraphObjectIdNode::SetupImpl(const FMovieGraphRenderPassSetupData& InSetupData)
-{
-	Super::SetupImpl(InSetupData);
-
-#if WITH_EDITOR
-	UEditorPerProjectUserSettings* EditorSettings = GetMutableDefault<UEditorPerProjectUserSettings>();
-	bPrevAllowSelectTranslucent = EditorSettings->bAllowSelectTranslucent;
-	EditorSettings->bAllowSelectTranslucent = bIncludeTranslucentObjects;
-#endif
-}
-
-void UMovieGraphObjectIdNode::TeardownImpl()
-{
-#if WITH_EDITOR
-	UEditorPerProjectUserSettings* EditorSettings = GetMutableDefault<UEditorPerProjectUserSettings>();
-	EditorSettings->bAllowSelectTranslucent = bPrevAllowSelectTranslucent;
-#endif
-	
-	Super::TeardownImpl();
-}
diff --git a/Engine/Plugins/MovieScene/MoviePipelineMaskRenderPass/Source/MoviePipelineMaskRenderPass/Private/Graph/MovieGraphObjectIdPass.cpp b/Engine/Plugins/MovieScene/MoviePipelineMaskRenderPass/Source/MoviePipelineMaskRenderPass/Private/Graph/MovieGraphObjectIdPass.cpp
--- a/Engine/Plugins/MovieScene/MoviePipelineMaskRenderPass/Source/MoviePipelineMaskRenderPass/Private/Graph/MovieGraphObjectIdPass.cpp
+++ b/Engine/Plugins/MovieScene/MoviePipelineMaskRenderPass/Source/MoviePipelineMaskRenderPass/Private/Graph/MovieGraphObjectIdPass.cpp
@@ -11,6 +11,10 @@
 #include "Graph/Nodes/MovieGraphCameraNode.h"
 #include "Graph/Nodes/MovieGraphImagePassBaseNode.h"
 
+#if WITH_EDITOR
+#include "Editor/EditorPerProjectUserSettings.h"
+#endif
+
 namespace UE::MovieGraph
 {
 	/**
@@ -194,6 +198,11 @@ void FMovieGraphObjectIdPass::Setup(TWeakObjectPtr<UMovieGraphDefaultRenderer> I
 
 	LayerData = InLayer;
 
+	if (const UMovieGraphObjectIdNode* ObjectIdNode = Cast<UMovieGraphObjectIdNode>(InLayer.RenderPassNode.Get()))
+	{
+		bIncludeTranslucentObjects = ObjectIdNode->bIncludeTranslucentObjects;
+	}
+
 	// We output three layers. This equals 6 "ranks" in the Cryptomatte spec.
 	for (int32 Index = 0; Index < 3; Index++)
 	{
@@ -251,7 +260,17 @@ UMovieGraphImagePassBaseNode* FMovieGraphObjectIdPass::GetParentNode(UMovieGraph
 
 void FMovieGraphObjectIdPass::Render(const FMovieGraphTraversalContext& InFrameTraversalContext, const FMovieGraphTimeStepData& InTimeData)
 {
+#if WITH_EDITOR
+	UEditorPerProjectUserSettings* EditorSettings = GetMutableDefault<UEditorPerProjectUserSettings>();
+	bPrevAllowSelectTranslucent = EditorSettings->bAllowSelectTranslucent;
+	EditorSettings->bAllowSelectTranslucent = bIncludeTranslucentObjects;
+#endif
+
 	FMovieGraphDeferredPass::Render(InFrameTraversalContext, InTimeData);
+
+#if WITH_EDITOR
+	EditorSettings->bAllowSelectTranslucent = bPrevAllowSelectTranslucent;
+#endif
 }
 
 TSharedRef<MoviePipeline::IMoviePipelineAccumulationArgs> FMovieGraphObjectIdPass::GetOrCreateAccumulator(TObjectPtr<UMovieGraphDefaultRenderer> InGraphRenderer, const UE::MovieGraph::FMovieGraphSampleState& InSampleState) const
diff --git a/Engine/Plugins/MovieScene/MoviePipelineMaskRenderPass/Source/MoviePipelineMaskRenderPass/Private/Graph/MovieGraphObjectIdPass.h b/Engine/Plugins/MovieScene/MoviePipelineMaskRenderPass/Source/MoviePipelineMaskRenderPass/Private/Graph/MovieGraphObjectIdPass.h
--- a/Engine/Plugins/MovieScene/MoviePipelineMaskRenderPass/Source/MoviePipelineMaskRenderPass/Private/Graph/MovieGraphObjectIdPass.h
+++ b/Engine/Plugins/MovieScene/MoviePipelineMaskRenderPass/Source/MoviePipelineMaskRenderPass/Private/Graph/MovieGraphObjectIdPass.h
@@ -49,6 +49,15 @@ protected:
 	/** The identifiers for all object ID layers that will be generated. */
 	TArray<FMovieGraphRenderDataIdentifier> RenderDataIdentifiers;
 
+	/** Whether translucent objects should be included in the IDs. */
+	bool bIncludeTranslucentObjects = false;
+
+	/**
+	 * This pass will modify Editor Settings to activate/deactivate translucent object selection (which bIncludeTranslucentObjects relies on). This
+	 * is the cached value of the property in Editor Settings (bAllowSelectTranslucent) before modifying it so it can be restored post-render.
+	 */
+	bool bPrevAllowSelectTranslucent = false;
+
 	/** ObjectID acceleration data that needs to be consistent throughout a render (cached per branch). */
 	inline static TMap<FName, UE::MoviePipeline::FObjectIdAccelerationData> AccelerationDataByBranch;
 };
\ No newline at end of file
diff --git a/Engine/Plugins/MovieScene/MoviePipelineMaskRenderPass/Source/MoviePipelineMaskRenderPass/Public/Graph/MovieGraphObjectIdNode.h b/Engine/Plugins/MovieScene/MoviePipelineMaskRenderPass/Source/MoviePipelineMaskRenderPass/Public/Graph/MovieGraphObjectIdNode.h
--- a/Engine/Plugins/MovieScene/MoviePipelineMaskRenderPass/Source/MoviePipelineMaskRenderPass/Public/Graph/MovieGraphObjectIdNode.h
+++ b/Engine/Plugins/MovieScene/MoviePipelineMaskRenderPass/Source/MoviePipelineMaskRenderPass/Public/Graph/MovieGraphObjectIdNode.h
@@ -49,11 +49,6 @@ protected:
 	UE_API virtual TUniquePtr<UE::MovieGraph::Rendering::FMovieGraphImagePassBase> CreateInstance() const override;
 	// ~UMovieGraphImagePassBaseNode Interface
 
-	// UMovieGraphRenderPassNode Interface
-	UE_API virtual void SetupImpl(const FMovieGraphRenderPassSetupData& InSetupData) override;
-	UE_API virtual void TeardownImpl() override;
-	// ~UMovieGraphRenderPassNode Interface
-
 public:
 	UPROPERTY(EditAnywhere, BlueprintReadWrite, Category = Overrides, meta = (InlineEditConditionToggle))
 	uint8 bOverride_SpatialSampleCount : 1;
@@ -78,13 +73,6 @@ public:
 	/** If true, translucent objects will be included in the ObjectId pass (but as an opaque layer due to limitations). False will omit translucent objects. */
 	UPROPERTY(EditAnywhere, BlueprintReadWrite, Category = "Settings", meta = (EditCondition="bOverride_bIncludeTranslucentObjects"))
 	bool bIncludeTranslucentObjects = false;
-
-private:
-	/**
-	 * The prior value of the project's bAllowSelectTranslucent setting. This node will temporarily change the value of this property to the value
-	 * set in bIncludeTranslucentObjects while rendering.
-	 */
-	bool bPrevAllowSelectTranslucent;
 };
 
 #undef UE_API
